- Moved the blocking portions (adb calls) of Android device detection off to a thread.
- This happens in a new module that the android target platforms depend on so that there aren't 6 different detection threads spun up.
- Thread isn't created if android sdk isn't found.

diff --git a/Engine/Source/Developer/Android/AndroidDeviceDetection/AndroidDeviceDetection.Build.cs b/Engine/Source/Developer/Android/AndroidDeviceDetection/AndroidDeviceDetection.Build.cs
new file mode 100644
--- /dev/null
+++ b/Engine/Source/Developer/Android/AndroidDeviceDetection/AndroidDeviceDetection.Build.cs
@@ -0,0 +1,30 @@
+// Copyright 1998-2014 Epic Games, Inc. All Rights Reserved.
+
+using UnrealBuildTool;
+
+public class AndroidDeviceDetection : ModuleRules
+{
+	public AndroidDeviceDetection( TargetInfo Target )
+	{
+		BinariesSubFolder = "Android";
+
+        PrivateDependencyModuleNames.AddRange(
+			new string[]
+			{
+				"Core",
+			}
+		);
+
+		PublicIncludePaths.AddRange(
+			new string[]
+			{
+				"Runtime/Core/Public/Android"
+			}
+		);
+
+		if (UEBuildConfiguration.bCompileAgainstEngine)
+		{
+			PrivateDependencyModuleNames.Add("Engine");
+		}
+	}
+}
diff --git a/Engine/Source/Developer/Android/AndroidDeviceDetection/Private/AndroidDeviceDetectionModule.cpp b/Engine/Source/Developer/Android/AndroidDeviceDetection/Private/AndroidDeviceDetectionModule.cpp
new file mode 100644
--- /dev/null
+++ b/Engine/Source/Developer/Android/AndroidDeviceDetection/Private/AndroidDeviceDetectionModule.cpp
@@ -0,0 +1,307 @@
+// Copyright 1998-2014 Epic Games, Inc. All Rights Reserved.
+
+/*=============================================================================
+	AndroidDeviceDetectionModule.cpp: Implements the FAndroidDeviceDetectionModule class.
+=============================================================================*/
+
+#include "AndroidDeviceDetectionPrivatePCH.h"
+
+#define LOCTEXT_NAMESPACE "FAndroidDeviceDetectionModule" 
+
+
+class FAndroidDeviceDetectionRunnable : public FRunnable
+{
+public:
+	FAndroidDeviceDetectionRunnable(const FString& InADBPath, TMap<FString,FAndroidDeviceInfo>& InDeviceMap, FCriticalSection* InDeviceMapLock) :
+		ADBPath(InADBPath),
+		StopTaskCounter(0),
+		DeviceMap(InDeviceMap),
+		DeviceMapLock(InDeviceMapLock)
+	{
+	}
+
+public:
+
+	// FRunnable interface.
+	virtual bool Init(void) 
+	{ 
+		return true; 
+	}
+
+	virtual void Exit(void) 
+	{
+	}
+
+	virtual void Stop(void)
+	{
+		StopTaskCounter.Increment();
+	}
+
+	virtual uint32 Run(void)
+	{
+		int LoopCount = 1;
+
+		while (StopTaskCounter.GetValue() == 0)
+		{
+			// query every 10 seconds
+			if (LoopCount++ >= 10)
+			{
+				QueryConnectedDevices();
+				LoopCount = 0;
+			}
+
+			FPlatformProcess::Sleep(1.0f);
+		}
+
+		return 0;
+	}
+
+private:
+
+	bool ExecuteAdbCommand( const FString& CommandLine, FString* OutStdOut, FString* OutStdErr ) const
+	{
+		// execute the command
+		int32 ReturnCode;
+		FString DefaultError;
+
+		// make sure there's a place for error output to go if the caller specified nullptr
+		if (!OutStdErr)
+		{
+			OutStdErr = &DefaultError;
+		}
+
+		FPlatformProcess::ExecProcess(*ADBPath, *CommandLine, &ReturnCode, OutStdOut, OutStdErr);
+
+		if (ReturnCode != 0)
+		{
+			FPlatformMisc::LowLevelOutputDebugStringf(TEXT("The Android SDK command '%s' failed to run. Return code: %d, Error: %s\n"), *CommandLine, ReturnCode, **OutStdErr);
+
+			return false;
+		}
+
+		return true;
+	}
+
+	void QueryConnectedDevices()
+	{
+		// grab the list of devices via adb
+		FString StdOut;
+		if (!ExecuteAdbCommand(TEXT("devices -l"), &StdOut, nullptr))
+		{
+			return;
+		}
+
+		// separate out each line
+		TArray<FString> DeviceStrings;
+		StdOut = StdOut.Replace(TEXT("\r"), TEXT("\n"));
+		StdOut.ParseIntoArray(&DeviceStrings, TEXT("\n"), true);
+
+		// a list containing all devices found this time, so we can remove anything not in this list
+		TArray<FString> CurrentlyConnectedDevices;
+
+		for (int32 StringIndex = 0; StringIndex < DeviceStrings.Num(); ++StringIndex)
+		{
+			const FString& DeviceString = DeviceStrings[StringIndex];
+
+			// skip over non-device lines
+			if (DeviceString.StartsWith("* ") || DeviceString.StartsWith("List "))
+			{
+				continue;
+			}
+
+			// grab the device serial number
+			int32 TabIndex;
+
+			if (!DeviceString.FindChar(TCHAR(' '), TabIndex))
+			{
+				continue;
+			}
+
+			FString SerialNumber = DeviceString.Left(TabIndex);
+
+			// add it to our list of currently connected devices
+			CurrentlyConnectedDevices.Add(SerialNumber);
+
+			// move on to next device if this one is already a known device
+			if (DeviceMap.Contains(SerialNumber))
+			{
+				continue;
+			}
+
+			// get the GL extensions string (and a bunch of other stuff)
+			FString GLESExtensions;
+
+			FString ExtensionsCommand = FString::Printf(TEXT("-s %s shell dumpsys SurfaceFlinger"), *SerialNumber);
+			if (!ExecuteAdbCommand(*ExtensionsCommand, &GLESExtensions, nullptr))
+			{
+				continue;
+			}
+
+			// grab the GL ES version
+			FString GLESVersionString;
+
+			FString VersionCommand = FString::Printf(TEXT("-s %s shell getprop ro.opengles.version"), *SerialNumber);
+			if (!ExecuteAdbCommand(*VersionCommand, &GLESVersionString, nullptr))
+			{
+				continue;
+			}
+
+			int GLESVersion = FCString::Atoi(*GLESVersionString);
+
+			// parse the device model
+			FString Model;
+			FParse::Value(*DeviceString, TEXT("model:"), Model);
+
+			// parse the device model
+			FString DeviceName;
+			FParse::Value(*DeviceString, TEXT("device:"), DeviceName);
+
+			// add the device to the map
+			{
+				FScopeLock ScopeLock(DeviceMapLock);
+
+				FAndroidDeviceInfo& DeviceInfo = DeviceMap.Add(SerialNumber);
+				DeviceInfo.SerialNumber = SerialNumber;
+				DeviceInfo.Model = Model;
+				DeviceInfo.DeviceName = DeviceName;
+				DeviceInfo.GLESExtenstions = GLESExtensions;
+				DeviceInfo.GLESVersion = GLESVersion;
+			}
+		}
+
+		// loop through the previously connected devices list and remove any that aren't still connected from the updated DeviceMap
+		TArray<FString> DevicesToRemove;
+
+		for (auto It = DeviceMap.CreateConstIterator(); It; ++It)
+		{
+			if (!CurrentlyConnectedDevices.Contains(It.Key()))
+			{
+				DevicesToRemove.Add(It.Key());
+			}
+		}
+
+		{
+			// enter the critical section and remove the devices from the map
+			FScopeLock ScopeLock(DeviceMapLock);
+
+			for (auto It = DevicesToRemove.CreateConstIterator(); It; ++It)
+			{
+				DeviceMap.Remove(*It);
+			}
+		}
+	}
+
+private:
+
+	// path to the adb command
+	FString ADBPath;
+
+	// > 0 if we've been asked to abort work in progress at the next opportunity
+	FThreadSafeCounter StopTaskCounter;
+
+	TMap<FString,FAndroidDeviceInfo>& DeviceMap;
+	FCriticalSection* DeviceMapLock;
+};
+
+class FAndroidDeviceDetection : public IAndroidDeviceDetection
+{
+public:
+
+	FAndroidDeviceDetection() :
+		DetectionThread(nullptr),
+		DetectionThreadRunnable(nullptr)
+	{
+		// get the SDK binaries folder
+		TCHAR AndroidDirectory[32768] = { 0 };
+		FPlatformMisc::GetEnvironmentVariable(TEXT("ANDROID_HOME"), AndroidDirectory, 32768);
+
+		if (AndroidDirectory[0] == 0)
+		{
+			return;
+		}
+
+		FString ADBPath = FString::Printf(TEXT("%s\\platform-tools\\adb.exe"), AndroidDirectory);
+
+		// if it doesn't exist, no SDK installed, so don't bother creating a thread to look for devices
+		if (!FPaths::FileExists(*ADBPath))
+		{
+			ADBPath.Empty();
+			return;
+		}
+
+		// create and fire off our device detection thread
+		DetectionThreadRunnable = new FAndroidDeviceDetectionRunnable(ADBPath, DeviceMap, &DeviceMapLock);
+		DetectionThread = FRunnableThread::Create(DetectionThreadRunnable, TEXT("FAndroidDeviceDetectionRunnable"), false);
+	}
+
+	virtual ~FAndroidDeviceDetection()
+	{
+		if (DetectionThreadRunnable && DetectionThread)
+		{
+			DetectionThreadRunnable->Stop();
+			DetectionThread->WaitForCompletion();
+		}
+	}
+
+	virtual const TMap<FString,FAndroidDeviceInfo>& GetDeviceMap() OVERRIDE
+	{
+		return DeviceMap;
+	}
+
+	virtual FCriticalSection* GetDeviceMapLock() OVERRIDE
+	{
+		return &DeviceMapLock;
+	}
+
+private:
+
+	FRunnableThread* DetectionThread;
+	FAndroidDeviceDetectionRunnable* DetectionThreadRunnable;
+
+	TMap<FString,FAndroidDeviceInfo> DeviceMap;
+	FCriticalSection DeviceMapLock;
+};
+
+
+/**
+ * Holds the target platform singleton.
+ */
+static FAndroidDeviceDetection* AndroidDeviceDetectionSingleton = nullptr;
+
+
+/**
+ * Module for detecting android devices.
+ */
+class FAndroidDeviceDetectionModule : public IAndroidDeviceDetectionModule
+{
+public:
+
+	/**
+	 * Destructor.
+	 */
+	~FAndroidDeviceDetectionModule( )
+	{
+		if (AndroidDeviceDetectionSingleton != nullptr)
+		{
+			delete AndroidDeviceDetectionSingleton;
+		}
+
+		AndroidDeviceDetectionSingleton = nullptr;
+	}
+
+	virtual IAndroidDeviceDetection* GetAndroidDeviceDetection() OVERRIDE
+	{
+		if (AndroidDeviceDetectionSingleton == nullptr)
+		{
+			AndroidDeviceDetectionSingleton = new FAndroidDeviceDetection();
+		}
+
+		return AndroidDeviceDetectionSingleton;
+	}
+};
+
+
+#undef LOCTEXT_NAMESPACE
+
+
+IMPLEMENT_MODULE( FAndroidDeviceDetectionModule, AndroidDeviceDetection);
diff --git a/Engine/Source/Developer/Android/AndroidDeviceDetection/Private/AndroidDeviceDetectionPrivatePCH.h b/Engine/Source/Developer/Android/AndroidDeviceDetection/Private/AndroidDeviceDetectionPrivatePCH.h
new file mode 100644
--- /dev/null
+++ b/Engine/Source/Developer/Android/AndroidDeviceDetection/Private/AndroidDeviceDetectionPrivatePCH.h
@@ -0,0 +1,21 @@
+// Copyright 1998-2014 Epic Games, Inc. All Rights Reserved.
+
+
+/*=============================================================================
+	AndroidDeviceDetectionPrivatePCH.h: Pre-compiled header file for the AndroidDeviceDetection module.
+=============================================================================*/
+
+#pragma once
+
+/* Dependencies
+ *****************************************************************************/
+
+#include "Core.h"
+#include "ModuleInterface.h"
+#include "ModuleManager.h"
+
+#if WITH_ENGINE
+	#include "Engine.h"
+#endif
+
+#include "AndroidDeviceDetection.h"
diff --git a/Engine/Source/Developer/Android/AndroidDeviceDetection/Public/AndroidDeviceDetection.h b/Engine/Source/Developer/Android/AndroidDeviceDetection/Public/AndroidDeviceDetection.h
new file mode 100644
--- /dev/null
+++ b/Engine/Source/Developer/Android/AndroidDeviceDetection/Public/AndroidDeviceDetection.h
@@ -0,0 +1,20 @@
+// Copyright 1998-2014 Epic Games, Inc. All Rights Reserved.
+
+/*=============================================================================
+	AndroidDeviceDetection.h: AndroidDeviceDetection module public header file.
+=============================================================================*/
+
+#pragma once
+
+
+/* Dependencies
+ *****************************************************************************/
+
+#include "Core.h"
+#include "ModuleManager.h"
+
+/* Interfaces
+ *****************************************************************************/
+
+#include "IAndroidDeviceDetection.h"
+#include "IAndroidDeviceDetectionModule.h"
diff --git a/Engine/Source/Developer/Android/AndroidDeviceDetection/Public/Interfaces/IAndroidDeviceDetection.h b/Engine/Source/Developer/Android/AndroidDeviceDetection/Public/Interfaces/IAndroidDeviceDetection.h
new file mode 100644
--- /dev/null
+++ b/Engine/Source/Developer/Android/AndroidDeviceDetection/Public/Interfaces/IAndroidDeviceDetection.h
@@ -0,0 +1,35 @@
+// Copyright 1998-2014 Epic Games, Inc. All Rights Reserved.
+
+/*=============================================================================
+	IAndroidDeviceDetection.h: Declares the IAndroidDeviceDetection interface.
+=============================================================================*/
+
+#pragma once
+
+
+struct FAndroidDeviceInfo
+{
+	FString SerialNumber;
+	FString Model;
+	FString DeviceName;
+	FString GLESExtenstions;		//@todo android: currently retrieved from SurfaceFlinger running 1.1, but so far seems ok for supported ES2 compression formats
+	int GLESVersion;
+};
+
+
+/**
+ * Interface for AndroidDeviceDetection module.
+ */
+class IAndroidDeviceDetection
+{
+public:
+	virtual const TMap<FString,FAndroidDeviceInfo>& GetDeviceMap() = 0;
+	virtual FCriticalSection* GetDeviceMapLock() = 0;
+
+protected:
+
+	/**
+	 * Virtual destructor
+	 */
+	virtual ~IAndroidDeviceDetection() { }
+};
diff --git a/Engine/Source/Developer/Android/AndroidDeviceDetection/Public/Interfaces/IAndroidDeviceDetectionModule.h b/Engine/Source/Developer/Android/AndroidDeviceDetection/Public/Interfaces/IAndroidDeviceDetectionModule.h
new file mode 100644
--- /dev/null
+++ b/Engine/Source/Developer/Android/AndroidDeviceDetection/Public/Interfaces/IAndroidDeviceDetectionModule.h
@@ -0,0 +1,27 @@
+// Copyright 1998-2014 Epic Games, Inc. All Rights Reserved.
+
+/*=============================================================================
+	IAndroidDeviceDetectionModule.h: Declares the IAndroidDeviceDetectionModule interface.
+=============================================================================*/
+
+#pragma once
+
+/**
+ * Interface for AndroidDeviceDetection module.
+ */
+class IAndroidDeviceDetectionModule
+	: public IModuleInterface
+{
+public:
+	/**
+	 * Returns the android device detection singleton.
+	 */
+	virtual IAndroidDeviceDetection* GetAndroidDeviceDetection() = 0;
+
+protected:
+
+	/**
+	 * Virtual destructor
+	 */
+	virtual ~IAndroidDeviceDetectionModule( ) { }
+};
diff --git a/Engine/Source/Developer/Android/AndroidTargetPlatform/AndroidTargetPlatform.Build.cs b/Engine/Source/Developer/Android/AndroidTargetPlatform/AndroidTargetPlatform.Build.cs
--- a/Engine/Source/Developer/Android/AndroidTargetPlatform/AndroidTargetPlatform.Build.cs
+++ b/Engine/Source/Developer/Android/AndroidTargetPlatform/AndroidTargetPlatform.Build.cs
@@ -17,6 +17,7 @@ public class AndroidTargetPlatform : ModuleRules
 				"TargetPlatform",
 				"Messaging",
 				"TargetDeviceServices",
+				"AndroidDeviceDetection",
 			}
 		);
 
diff --git a/Engine/Source/Developer/Android/AndroidTargetPlatform/Private/AndroidTargetPlatform.h b/Engine/Source/Developer/Android/AndroidTargetPlatform/Private/AndroidTargetPlatform.h
--- a/Engine/Source/Developer/Android/AndroidTargetPlatform/Private/AndroidTargetPlatform.h
+++ b/Engine/Source/Developer/Android/AndroidTargetPlatform/Private/AndroidTargetPlatform.h
@@ -164,30 +164,11 @@ protected:
 		return true;
 	}
 
-	/**
-	 * Executes an SDK command with the specified command line.
-	 *
-	 * @param Params - The command line parameters.
-	 * @param OutStdOut - Optional pointer to a string that will hold the command's output log.
-	 * @param OutStdErr - Optional pointer to a string that will hold the error message, if any.
-	 *
-	 * @return true on success, false otherwise.
-	 */
-	bool ExecuteAdbCommand( const FString& Params, FString* OutStdOut, FString* OutStdErr ) const;
-
-	/**
-	 * Gathers a list of all devices that are currently connected.
-	 */
-	void QueryConnectedDevices( );
-
 private:
 
 	// Handles when the ticker fires.
 	bool HandleTicker( float DeltaTime );
 
-	// Handles received pong messages from the LauncherDaemon.
-//	void HandlePongMessage( const FAndroidPong &Message, const IMessageContextRef &Context );
-
 private:
 
 	// Holds a map of valid devices.
@@ -195,6 +176,9 @@ private:
 
 	// Holds a delegate to be invoked when the widget ticks.
 	FTickerDelegate TickDelegate;
+	
+	// Pointer to the device detection handler that grabs device ids in another thread
+	IAndroidDeviceDetection* DeviceDetection;
 
 #if WITH_ENGINE
 	// Holds the Engine INI settings (for quick access).
diff --git a/Engine/Source/Developer/Android/AndroidTargetPlatform/Private/AndroidTargetPlatform.inl b/Engine/Source/Developer/Android/AndroidTargetPlatform/Private/AndroidTargetPlatform.inl
--- a/Engine/Source/Developer/Android/AndroidTargetPlatform/Private/AndroidTargetPlatform.inl
+++ b/Engine/Source/Developer/Android/AndroidTargetPlatform/Private/AndroidTargetPlatform.inl
@@ -8,7 +8,8 @@
 /* FAndroidTargetPlatform structors
  *****************************************************************************/
 
-inline FAndroidTargetPlatform::FAndroidTargetPlatform( )
+inline FAndroidTargetPlatform::FAndroidTargetPlatform( ) :
+	DeviceDetection(nullptr)
 {
 	#if WITH_ENGINE
 		FConfigCacheIni::LoadLocalIniFile(EngineSettings, TEXT("Engine"), true, *PlatformName());
@@ -17,7 +18,7 @@ inline FAndroidTargetPlatform::FAndroidTargetPlatform( )
 	#endif
 
 	TickDelegate = FTickerDelegate::CreateRaw(this, &FAndroidTargetPlatform::HandleTicker);
-	FTicker::GetCoreTicker().AddTicker(TickDelegate, 5.0f);
+	FTicker::GetCoreTicker().AddTicker(TickDelegate, 4.0f);
 }
 
 
@@ -297,137 +298,54 @@ inline void FAndroidTargetPlatform::AddTextureFormatIfSupports( FName Format, TA
 }
 
 
-inline bool FAndroidTargetPlatform::ExecuteAdbCommand( const FString& CommandLine, FString* OutStdOut, FString* OutStdErr ) const
-{
-	// get the SDK binaries folder
-	TCHAR AndroidDirectory[32768] = { 0 };
-	FPlatformMisc::GetEnvironmentVariable(TEXT("ANDROID_HOME"), AndroidDirectory, 32768);
-
-	if (AndroidDirectory[0] == 0)
-	{
-		return false;
-	}
-
-	FString Filename = FString::Printf(TEXT("%s\\platform-tools\\adb.exe"), AndroidDirectory);
-
-	// execute the command
-	int32 ReturnCode;
-	FString DefaultError;
-
-	// make sure there's a place for error output to go if the caller specified nullptr
-	if (!OutStdErr)
-	{
-		OutStdErr = &DefaultError;
-	}
-
-	FPlatformProcess::ExecProcess(*Filename, *CommandLine, &ReturnCode, OutStdOut, OutStdErr);
-
-	if (ReturnCode != 0)
-	{
-		FPlatformMisc::LowLevelOutputDebugStringf(TEXT("The Android SDK command '%s' failed to run. Return code: %d, Error: %s\n"), *CommandLine, ReturnCode, **OutStdErr);
-
-		return false;
-	}
-
-	return true;
-}
-
+/* FAndroidTargetPlatform callbacks
+ *****************************************************************************/
 
-inline void FAndroidTargetPlatform::QueryConnectedDevices( )
+inline bool FAndroidTargetPlatform::HandleTicker( float DeltaTime )
 {
-	FString StdOut;
-
-	// grab the list of devices via adb
-	if (!ExecuteAdbCommand(TEXT("devices -l"), &StdOut, nullptr))
+	if (DeviceDetection == nullptr)
 	{
-		return;
+		DeviceDetection = FModuleManager::LoadModuleChecked<IAndroidDeviceDetectionModule>("AndroidDeviceDetection").GetAndroidDeviceDetection();
 	}
 
-	// separate out each line
-	TArray<FString> DeviceStrings;
-	StdOut = StdOut.Replace(TEXT("\r"), TEXT("\n"));
-	StdOut.ParseIntoArray(&DeviceStrings, TEXT("\n"), true);
-
 	TArray<FString> ConnectedDeviceIds;
 
-	for (int32 StringIndex = 0; StringIndex < DeviceStrings.Num(); ++StringIndex)
 	{
-		const FString& DeviceString = DeviceStrings[StringIndex];
-
-		// skip over non-device lines
-
-		if (DeviceString.StartsWith("* ") || DeviceString.StartsWith("List "))
-		{
-			continue;
-		}
+		FScopeLock ScopeLock(DeviceDetection->GetDeviceMapLock());
 
-		// grab the device serial number
-		int32 TabIndex;
-
-		if (!DeviceString.FindChar(TCHAR(' '), TabIndex))
+		auto DeviceIt = DeviceDetection->GetDeviceMap().CreateConstIterator();
+		
+		for (; DeviceIt; ++DeviceIt)
 		{
-			continue;
-		}
+			ConnectedDeviceIds.Add(DeviceIt.Key());
 
-		FString SerialNumber = DeviceString.Left(TabIndex);
-		ConnectedDeviceIds.Add(DeviceString.Left(TabIndex));
+			// see if this device is already known
+			if (Devices.Contains(DeviceIt.Key()))
+			{
+				continue;
+			}
 
-		// move on to next device if this one is already a known device
-		if (Devices.Find(SerialNumber))
-		{
-			continue;
-		}
+			const FAndroidDeviceInfo& DeviceInfo = DeviceIt.Value();
 
-		// get the GL extensions string (and a bunch of other stuff)
-		FString ExtensionsString;
+#if 0	// todo android: really need to filter devices differently based on deploy vs. package
+			// check if this platform is supported by the extensions and version
+			if (!SupportedByExtensionsString(DeviceInfo.GLESExtensions, DeviceInfo.GLESVersion))
+			{
+				continue;
+			}
+#endif
 
-		FString ExtensionsCommand = FString::Printf(TEXT("-s %s shell dumpsys SurfaceFlinger"), *SerialNumber);
-		if (!ExecuteAdbCommand(*ExtensionsCommand, &ExtensionsString, nullptr))
-		{
-			continue;
-		}
+			// create target device
+			FAndroidTargetDevicePtr& Device = Devices.Add(DeviceInfo.SerialNumber);
 
-		// grab the GL ES version
-		FString GLESVersionString;
+			Device = MakeShareable(new FAndroidTargetDevice(*this, DeviceInfo.SerialNumber, GetAndroidVariantName()));
 
-		FString VersionCommand = FString::Printf(TEXT("-s %s shell getprop ro.opengles.version"), *SerialNumber);
-		if (!ExecuteAdbCommand(*VersionCommand, &GLESVersionString, nullptr))
-		{
-			continue;
-		}
+			Device->SetConnected(true);
+			Device->SetModel(DeviceInfo.Model);
+			Device->SetDeviceName(DeviceInfo.DeviceName);
 
-		int GLESVersion = FCString::Atoi(*GLESVersionString);
-
-		// check if this platform is supported by the extensions and version
-		if (!SupportedByExtensionsString(ExtensionsString, GLESVersion))
-		{
-			continue;
+			DeviceDiscoveredEvent.Broadcast(Device.ToSharedRef());
 		}
-
-		// create target device
-		FAndroidTargetDevicePtr& Device = Devices.Add(SerialNumber);
-
-		Device = MakeShareable(new FAndroidTargetDevice(*this, SerialNumber, GetAndroidVariantName()));
-
-		Device->SetConnected(true);
-
-		// parse device model
-		FString Model;
-
-		if (FParse::Value(*DeviceString, TEXT("model:"), Model))
-		{
-			Device->SetModel(Model);
-		}
-	
-		// parse device name
-		FString DeviceName;
-
-		if (FParse::Value(*DeviceString, TEXT("device:"), DeviceName))
-		{
-			Device->SetDeviceName(DeviceName);
-		}
-
-		DeviceDiscoveredEvent.Broadcast(Device.ToSharedRef());
 	}
 
 	// remove disconnected devices
@@ -443,15 +361,6 @@ inline void FAndroidTargetPlatform::QueryConnectedDevices( )
 			DeviceLostEvent.Broadcast(RemovedDevice.ToSharedRef());
 		}
 	}
-}
-
-
-/* FAndroidTargetPlatform callbacks
- *****************************************************************************/
-
-inline bool FAndroidTargetPlatform::HandleTicker( float DeltaTime )
-{
-	QueryConnectedDevices();
 
 	return true;
 }
diff --git a/Engine/Source/Developer/Android/AndroidTargetPlatform/Private/AndroidTargetPlatformPrivatePCH.h b/Engine/Source/Developer/Android/AndroidTargetPlatform/Private/AndroidTargetPlatformPrivatePCH.h
--- a/Engine/Source/Developer/Android/AndroidTargetPlatform/Private/AndroidTargetPlatformPrivatePCH.h
+++ b/Engine/Source/Developer/Android/AndroidTargetPlatform/Private/AndroidTargetPlatformPrivatePCH.h
@@ -22,6 +22,7 @@
 #endif
 #include "TargetDeviceServices.h"
 #include "TargetPlatform.h"
+#include "AndroidDeviceDetection.h"
 
 
 /* Private includes
diff --git a/Engine/Source/Developer/Android/Android_ATCTargetPlatform/Android_ATCTargetPlatform.Build.cs b/Engine/Source/Developer/Android/Android_ATCTargetPlatform/Android_ATCTargetPlatform.Build.cs
--- a/Engine/Source/Developer/Android/Android_ATCTargetPlatform/Android_ATCTargetPlatform.Build.cs
+++ b/Engine/Source/Developer/Android/Android_ATCTargetPlatform/Android_ATCTargetPlatform.Build.cs
@@ -16,6 +16,7 @@ public class Android_ATCTargetPlatform : ModuleRules
 				"TargetPlatform",
 				"Messaging",
 				"TargetDeviceServices",
+				"AndroidDeviceDetection",
 			}
 		);
 
diff --git a/Engine/Source/Developer/Android/Android_ATCTargetPlatform/Private/Android_ATCTargetPlatformPrivatePCH.h b/Engine/Source/Developer/Android/Android_ATCTargetPlatform/Private/Android_ATCTargetPlatformPrivatePCH.h
--- a/Engine/Source/Developer/Android/Android_ATCTargetPlatform/Private/Android_ATCTargetPlatformPrivatePCH.h
+++ b/Engine/Source/Developer/Android/Android_ATCTargetPlatform/Private/Android_ATCTargetPlatformPrivatePCH.h
@@ -22,6 +22,7 @@
 #endif
 #include "TargetDeviceServices.h"
 #include "TargetPlatform.h"
+#include "AndroidDeviceDetection.h"
 
 
 /* Private includes
diff --git a/Engine/Source/Developer/Android/Android_DXTTargetPlatform/Android_DXTTargetPlatform.Build.cs b/Engine/Source/Developer/Android/Android_DXTTargetPlatform/Android_DXTTargetPlatform.Build.cs
--- a/Engine/Source/Developer/Android/Android_DXTTargetPlatform/Android_DXTTargetPlatform.Build.cs
+++ b/Engine/Source/Developer/Android/Android_DXTTargetPlatform/Android_DXTTargetPlatform.Build.cs
@@ -16,6 +16,7 @@ public class Android_DXTTargetPlatform : ModuleRules
 				"TargetPlatform",
 				"Messaging",
 				"TargetDeviceServices",
+				"AndroidDeviceDetection",
 			}
 		);
 
diff --git a/Engine/Source/Developer/Android/Android_DXTTargetPlatform/Private/Android_DXTTargetPlatformPrivatePCH.h b/Engine/Source/Developer/Android/Android_DXTTargetPlatform/Private/Android_DXTTargetPlatformPrivatePCH.h
--- a/Engine/Source/Developer/Android/Android_DXTTargetPlatform/Private/Android_DXTTargetPlatformPrivatePCH.h
+++ b/Engine/Source/Developer/Android/Android_DXTTargetPlatform/Private/Android_DXTTargetPlatformPrivatePCH.h
@@ -22,6 +22,7 @@
 #endif
 #include "TargetDeviceServices.h"
 #include "TargetPlatform.h"
+#include "AndroidDeviceDetection.h"
 
 
 /* Private includes
diff --git a/Engine/Source/Developer/Android/Android_ETC1TargetPlatform/Android_ETC1TargetPlatform.Build.cs b/Engine/Source/Developer/Android/Android_ETC1TargetPlatform/Android_ETC1TargetPlatform.Build.cs
--- a/Engine/Source/Developer/Android/Android_ETC1TargetPlatform/Android_ETC1TargetPlatform.Build.cs
+++ b/Engine/Source/Developer/Android/Android_ETC1TargetPlatform/Android_ETC1TargetPlatform.Build.cs
@@ -16,6 +16,7 @@ public class Android_ETC1TargetPlatform : ModuleRules
 				"TargetPlatform",
 				"Messaging",
 				"TargetDeviceServices",
+				"AndroidDeviceDetection",
 			}
 		);
 
diff --git a/Engine/Source/Developer/Android/Android_ETC1TargetPlatform/Private/Android_ETC1TargetPlatformPrivatePCH.h b/Engine/Source/Developer/Android/Android_ETC1TargetPlatform/Private/Android_ETC1TargetPlatformPrivatePCH.h
--- a/Engine/Source/Developer/Android/Android_ETC1TargetPlatform/Private/Android_ETC1TargetPlatformPrivatePCH.h
+++ b/Engine/Source/Developer/Android/Android_ETC1TargetPlatform/Private/Android_ETC1TargetPlatformPrivatePCH.h
@@ -22,6 +22,7 @@
 #endif
 #include "TargetDeviceServices.h"
 #include "TargetPlatform.h"
+#include "AndroidDeviceDetection.h"
 
 
 /* Private includes
diff --git a/Engine/Source/Developer/Android/Android_ETC2TargetPlatform/Android_ETC2TargetPlatform.Build.cs b/Engine/Source/Developer/Android/Android_ETC2TargetPlatform/Android_ETC2TargetPlatform.Build.cs
--- a/Engine/Source/Developer/Android/Android_ETC2TargetPlatform/Android_ETC2TargetPlatform.Build.cs
+++ b/Engine/Source/Developer/Android/Android_ETC2TargetPlatform/Android_ETC2TargetPlatform.Build.cs
@@ -16,6 +16,7 @@ public class Android_ETC2TargetPlatform : ModuleRules
 				"TargetPlatform",
 				"Messaging",
 				"TargetDeviceServices",
+				"AndroidDeviceDetection",
 			}
 		);
 
diff --git a/Engine/Source/Developer/Android/Android_ETC2TargetPlatform/Private/Android_ETC2TargetPlatformPrivatePCH.h b/Engine/Source/Developer/Android/Android_ETC2TargetPlatform/Private/Android_ETC2TargetPlatformPrivatePCH.h
--- a/Engine/Source/Developer/Android/Android_ETC2TargetPlatform/Private/Android_ETC2TargetPlatformPrivatePCH.h
+++ b/Engine/Source/Developer/Android/Android_ETC2TargetPlatform/Private/Android_ETC2TargetPlatformPrivatePCH.h
@@ -22,6 +22,7 @@
 #endif
 #include "TargetDeviceServices.h"
 #include "TargetPlatform.h"
+#include "AndroidDeviceDetection.h"
 
 
 /* Private includes
diff --git a/Engine/Source/Developer/Android/Android_PVRTCTargetPlatform/Android_PVRTCTargetPlatform.Build.cs b/Engine/Source/Developer/Android/Android_PVRTCTargetPlatform/Android_PVRTCTargetPlatform.Build.cs
--- a/Engine/Source/Developer/Android/Android_PVRTCTargetPlatform/Android_PVRTCTargetPlatform.Build.cs
+++ b/Engine/Source/Developer/Android/Android_PVRTCTargetPlatform/Android_PVRTCTargetPlatform.Build.cs
@@ -16,6 +16,7 @@ public class Android_PVRTCTargetPlatform : ModuleRules
 				"TargetPlatform",
 				"Messaging",
 				"TargetDeviceServices",
+				"AndroidDeviceDetection",
 			}
 		);
 
diff --git a/Engine/Source/Developer/Android/Android_PVRTCTargetPlatform/Private/Android_PVRTCTargetPlatformPrivatePCH.h b/Engine/Source/Developer/Android/Android_PVRTCTargetPlatform/Private/Android_PVRTCTargetPlatformPrivatePCH.h
--- a/Engine/Source/Developer/Android/Android_PVRTCTargetPlatform/Private/Android_PVRTCTargetPlatformPrivatePCH.h
+++ b/Engine/Source/Developer/Android/Android_PVRTCTargetPlatform/Private/Android_PVRTCTargetPlatformPrivatePCH.h
@@ -22,6 +22,7 @@
 #endif
 #include "TargetDeviceServices.h"
 #include "TargetPlatform.h"
+#include "AndroidDeviceDetection.h"
 
 
 /* Private includes
