- enabled device filtering by texture format support

diff --git a/Engine/Source/Developer/Android/AndroidDeviceDetection/Private/AndroidDeviceDetectionModule.cpp b/Engine/Source/Developer/Android/AndroidDeviceDetection/Private/AndroidDeviceDetectionModule.cpp
--- a/Engine/Source/Developer/Android/AndroidDeviceDetection/Private/AndroidDeviceDetectionModule.cpp
+++ b/Engine/Source/Developer/Android/AndroidDeviceDetection/Private/AndroidDeviceDetectionModule.cpp
@@ -39,7 +39,7 @@ public:
 
 	virtual uint32 Run(void)
 	{
-		int LoopCount = 1;
+		int LoopCount = 10;
 
 		while (StopTaskCounter.GetValue() == 0)
 		{
@@ -164,7 +164,7 @@ private:
 				DeviceInfo.SerialNumber = SerialNumber;
 				DeviceInfo.Model = Model;
 				DeviceInfo.DeviceName = DeviceName;
-				DeviceInfo.GLESExtenstions = GLESExtensions;
+				DeviceInfo.GLESExtensions = GLESExtensions;
 				DeviceInfo.GLESVersion = GLESVersion;
 			}
 		}
diff --git a/Engine/Source/Developer/Android/AndroidDeviceDetection/Public/Interfaces/IAndroidDeviceDetection.h b/Engine/Source/Developer/Android/AndroidDeviceDetection/Public/Interfaces/IAndroidDeviceDetection.h
--- a/Engine/Source/Developer/Android/AndroidDeviceDetection/Public/Interfaces/IAndroidDeviceDetection.h
+++ b/Engine/Source/Developer/Android/AndroidDeviceDetection/Public/Interfaces/IAndroidDeviceDetection.h
@@ -12,7 +12,7 @@ struct FAndroidDeviceInfo
 	FString SerialNumber;
 	FString Model;
 	FString DeviceName;
-	FString GLESExtenstions;		//@todo android: currently retrieved from SurfaceFlinger running 1.1, but so far seems ok for supported ES2 compression formats
+	FString GLESExtensions;		//@todo android: currently retrieved from SurfaceFlinger running 1.1, but so far seems ok for supported ES2 compression formats
 	int GLESVersion;
 };
 
diff --git a/Engine/Source/Developer/Android/AndroidTargetPlatform/Private/AndroidTargetPlatform.inl b/Engine/Source/Developer/Android/AndroidTargetPlatform/Private/AndroidTargetPlatform.inl
--- a/Engine/Source/Developer/Android/AndroidTargetPlatform/Private/AndroidTargetPlatform.inl
+++ b/Engine/Source/Developer/Android/AndroidTargetPlatform/Private/AndroidTargetPlatform.inl
@@ -327,13 +327,11 @@ inline bool FAndroidTargetPlatform::HandleTicker( float DeltaTime )
 
 			const FAndroidDeviceInfo& DeviceInfo = DeviceIt.Value();
 
-#if 0	// todo android: really need to filter devices differently based on deploy vs. package
 			// check if this platform is supported by the extensions and version
 			if (!SupportedByExtensionsString(DeviceInfo.GLESExtensions, DeviceInfo.GLESVersion))
 			{
 				continue;
 			}
-#endif
 
 			// create target device
 			FAndroidTargetDevicePtr& Device = Devices.Add(DeviceInfo.SerialNumber);
