- Fix build break.  Missing pure virtual in android target platform.

diff --git a/Engine/Source/Developer/Android/AndroidTargetPlatform/Private/AndroidTargetPlatform.h b/Engine/Source/Developer/Android/AndroidTargetPlatform/Private/AndroidTargetPlatform.h
--- a/Engine/Source/Developer/Android/AndroidTargetPlatform/Private/AndroidTargetPlatform.h
+++ b/Engine/Source/Developer/Android/AndroidTargetPlatform/Private/AndroidTargetPlatform.h
@@ -73,6 +73,8 @@ public:
 
 	// Begin ITargetPlatform interface
 
+	virtual void EnableDeviceCheck(bool OnOff) OVERRIDE {}
+
 	virtual bool AddDevice( const FString& DeviceName, bool bDefault ) OVERRIDE
 	{
 		return false;
