removed unnecessary module dependencies from Android target platforms

diff --git a/Engine/Source/Developer/Android/AndroidPlatformEditor/AndroidPlatformEditor.Build.cs b/Engine/Source/Developer/Android/AndroidPlatformEditor/AndroidPlatformEditor.Build.cs
--- a/Engine/Source/Developer/Android/AndroidPlatformEditor/AndroidPlatformEditor.Build.cs
+++ b/Engine/Source/Developer/Android/AndroidPlatformEditor/AndroidPlatformEditor.Build.cs
@@ -9,8 +9,7 @@ public class AndroidPlatformEditor : ModuleRules
 		BinariesSubFolder = "Android";
 
 		PrivateDependencyModuleNames.AddRange(
-			new string[]
-			{
+			new string[] {
 				"Core",
 				"CoreUObject",
 				"InputCore",
@@ -19,12 +18,17 @@ public class AndroidPlatformEditor : ModuleRules
 				"EditorStyle",
 				"PropertyEditor",
 				"SharedSettingsWidgets",
-				"Settings",
 				"SourceControl",
 				"AndroidRuntimeSettings",
 			}
 		);
 
+		PrivateIncludePathModuleNames.AddRange(
+			new string[] {
+				"Settings",
+			}
+		);
+
         // this is listed above, so it isn't really dynamically loaded, this just marks it as being platform specific.
 		PlatformSpecificDynamicallyLoadedModuleNames.Add("AndroidRuntimeSettings");
 	}
diff --git a/Engine/Source/Developer/Android/AndroidTargetPlatform/AndroidTargetPlatform.Build.cs b/Engine/Source/Developer/Android/AndroidTargetPlatform/AndroidTargetPlatform.Build.cs
--- a/Engine/Source/Developer/Android/AndroidTargetPlatform/AndroidTargetPlatform.Build.cs
+++ b/Engine/Source/Developer/Android/AndroidTargetPlatform/AndroidTargetPlatform.Build.cs
@@ -13,10 +13,7 @@ public class AndroidTargetPlatform : ModuleRules
 			{
 				"Core",
 				"CoreUObject",
-				"Sockets",
 				"TargetPlatform",
-				"Messaging",
-				"TargetDeviceServices",
 				"AndroidDeviceDetection",
 			}
 		);
diff --git a/Engine/Source/Developer/Android/AndroidTargetPlatform/Private/AndroidTargetPlatformPrivatePCH.h b/Engine/Source/Developer/Android/AndroidTargetPlatform/Private/AndroidTargetPlatformPrivatePCH.h
--- a/Engine/Source/Developer/Android/AndroidTargetPlatform/Private/AndroidTargetPlatformPrivatePCH.h
+++ b/Engine/Source/Developer/Android/AndroidTargetPlatform/Private/AndroidTargetPlatformPrivatePCH.h
@@ -20,7 +20,7 @@
 	#include "Engine.h"
 	#include "TextureCompressorModule.h"
 #endif
-#include "TargetDeviceServices.h"
+
 #include "TargetPlatform.h"
 #include "AndroidDeviceDetection.h"
 
diff --git a/Engine/Source/Developer/Android/Android_ATCTargetPlatform/Android_ATCTargetPlatform.Build.cs b/Engine/Source/Developer/Android/Android_ATCTargetPlatform/Android_ATCTargetPlatform.Build.cs
--- a/Engine/Source/Developer/Android/Android_ATCTargetPlatform/Android_ATCTargetPlatform.Build.cs
+++ b/Engine/Source/Developer/Android/Android_ATCTargetPlatform/Android_ATCTargetPlatform.Build.cs
@@ -12,10 +12,7 @@ public class Android_ATCTargetPlatform : ModuleRules
 			new string[]
 			{
 				"Core",
-				"Sockets",
 				"TargetPlatform",
-				"Messaging",
-				"TargetDeviceServices",
 				"AndroidDeviceDetection",
 			}
 		);
diff --git a/Engine/Source/Developer/Android/Android_ATCTargetPlatform/Private/Android_ATCTargetPlatformPrivatePCH.h b/Engine/Source/Developer/Android/Android_ATCTargetPlatform/Private/Android_ATCTargetPlatformPrivatePCH.h
--- a/Engine/Source/Developer/Android/Android_ATCTargetPlatform/Private/Android_ATCTargetPlatformPrivatePCH.h
+++ b/Engine/Source/Developer/Android/Android_ATCTargetPlatform/Private/Android_ATCTargetPlatformPrivatePCH.h
@@ -13,14 +13,13 @@
 #include "Core.h"
 #include "ModuleInterface.h"
 #include "ModuleManager.h"
-#include "Messaging.h"
 #include "Runtime/Core/Public/Android/AndroidProperties.h"
 
 #if WITH_ENGINE
 	#include "Engine.h"
 	#include "TextureCompressorModule.h"
 #endif
-#include "TargetDeviceServices.h"
+
 #include "TargetPlatform.h"
 #include "AndroidDeviceDetection.h"
 
diff --git a/Engine/Source/Developer/Android/Android_DXTTargetPlatform/Android_DXTTargetPlatform.Build.cs b/Engine/Source/Developer/Android/Android_DXTTargetPlatform/Android_DXTTargetPlatform.Build.cs
--- a/Engine/Source/Developer/Android/Android_DXTTargetPlatform/Android_DXTTargetPlatform.Build.cs
+++ b/Engine/Source/Developer/Android/Android_DXTTargetPlatform/Android_DXTTargetPlatform.Build.cs
@@ -12,10 +12,7 @@ public class Android_DXTTargetPlatform : ModuleRules
 			new string[]
 			{
 				"Core",
-				"Sockets",
 				"TargetPlatform",
-				"Messaging",
-				"TargetDeviceServices",
 				"AndroidDeviceDetection",
 			}
 		);
diff --git a/Engine/Source/Developer/Android/Android_DXTTargetPlatform/Private/Android_DXTTargetPlatformPrivatePCH.h b/Engine/Source/Developer/Android/Android_DXTTargetPlatform/Private/Android_DXTTargetPlatformPrivatePCH.h
--- a/Engine/Source/Developer/Android/Android_DXTTargetPlatform/Private/Android_DXTTargetPlatformPrivatePCH.h
+++ b/Engine/Source/Developer/Android/Android_DXTTargetPlatform/Private/Android_DXTTargetPlatformPrivatePCH.h
@@ -13,14 +13,13 @@
 #include "Core.h"
 #include "ModuleInterface.h"
 #include "ModuleManager.h"
-#include "Messaging.h"
 #include "Runtime/Core/Public/Android/AndroidProperties.h"
 
 #if WITH_ENGINE
 	#include "Engine.h"
 	#include "TextureCompressorModule.h"
 #endif
-#include "TargetDeviceServices.h"
+
 #include "TargetPlatform.h"
 #include "AndroidDeviceDetection.h"
 
diff --git a/Engine/Source/Developer/Android/Android_ETC1TargetPlatform/Android_ETC1TargetPlatform.Build.cs b/Engine/Source/Developer/Android/Android_ETC1TargetPlatform/Android_ETC1TargetPlatform.Build.cs
--- a/Engine/Source/Developer/Android/Android_ETC1TargetPlatform/Android_ETC1TargetPlatform.Build.cs
+++ b/Engine/Source/Developer/Android/Android_ETC1TargetPlatform/Android_ETC1TargetPlatform.Build.cs
@@ -8,36 +8,34 @@ public class Android_ETC1TargetPlatform : ModuleRules
 	{
 		BinariesSubFolder = "Android";
 
-        PrivateDependencyModuleNames.AddRange(
+		PublicIncludePaths.AddRange(
+			new string[]
+			{
+				"Runtime/Core/Public/Android"
+			}
+		);
+
+		PrivateDependencyModuleNames.AddRange(
 			new string[]
 			{
 				"Core",
-				"Sockets",
 				"TargetPlatform",
-				"Messaging",
-				"TargetDeviceServices",
 				"AndroidDeviceDetection",
 			}
 		);
 
-		PublicIncludePaths.AddRange(
+		PrivateIncludePaths.AddRange(
 			new string[]
 			{
-				"Runtime/Core/Public/Android"
+				"Developer/Android/AndroidTargetPlatform/Private",				
 			}
 		);
 
+		// compile with Engine
 		if (UEBuildConfiguration.bCompileAgainstEngine)
 		{
 			PrivateDependencyModuleNames.Add("Engine");
 			PrivateIncludePathModuleNames.Add("TextureCompressor");		//@todo android: AndroidTargetPlatform.Build
 		}
-
-		PrivateIncludePaths.AddRange(
-			new string[]
-			{
-				"Developer/Android/AndroidTargetPlatform/Private",				
-			}
-		);
 	}
 }
\ No newline at end of file
diff --git a/Engine/Source/Developer/Android/Android_ETC1TargetPlatform/Private/Android_ETC1TargetPlatformPrivatePCH.h b/Engine/Source/Developer/Android/Android_ETC1TargetPlatform/Private/Android_ETC1TargetPlatformPrivatePCH.h
--- a/Engine/Source/Developer/Android/Android_ETC1TargetPlatform/Private/Android_ETC1TargetPlatformPrivatePCH.h
+++ b/Engine/Source/Developer/Android/Android_ETC1TargetPlatform/Private/Android_ETC1TargetPlatformPrivatePCH.h
@@ -13,14 +13,13 @@
 #include "Core.h"
 #include "ModuleInterface.h"
 #include "ModuleManager.h"
-#include "Messaging.h"
 #include "Runtime/Core/Public/Android/AndroidProperties.h"
 
 #if WITH_ENGINE
 	#include "Engine.h"
 	#include "TextureCompressorModule.h"
 #endif
-#include "TargetDeviceServices.h"
+
 #include "TargetPlatform.h"
 #include "AndroidDeviceDetection.h"
 
diff --git a/Engine/Source/Developer/Android/Android_ETC2TargetPlatform/Android_ETC2TargetPlatform.Build.cs b/Engine/Source/Developer/Android/Android_ETC2TargetPlatform/Android_ETC2TargetPlatform.Build.cs
--- a/Engine/Source/Developer/Android/Android_ETC2TargetPlatform/Android_ETC2TargetPlatform.Build.cs
+++ b/Engine/Source/Developer/Android/Android_ETC2TargetPlatform/Android_ETC2TargetPlatform.Build.cs
@@ -12,10 +12,7 @@ public class Android_ETC2TargetPlatform : ModuleRules
 			new string[]
 			{
 				"Core",
-				"Sockets",
 				"TargetPlatform",
-				"Messaging",
-				"TargetDeviceServices",
 				"AndroidDeviceDetection",
 			}
 		);
diff --git a/Engine/Source/Developer/Android/Android_ETC2TargetPlatform/Private/Android_ETC2TargetPlatformPrivatePCH.h b/Engine/Source/Developer/Android/Android_ETC2TargetPlatform/Private/Android_ETC2TargetPlatformPrivatePCH.h
--- a/Engine/Source/Developer/Android/Android_ETC2TargetPlatform/Private/Android_ETC2TargetPlatformPrivatePCH.h
+++ b/Engine/Source/Developer/Android/Android_ETC2TargetPlatform/Private/Android_ETC2TargetPlatformPrivatePCH.h
@@ -13,14 +13,13 @@
 #include "Core.h"
 #include "ModuleInterface.h"
 #include "ModuleManager.h"
-#include "Messaging.h"
 #include "Runtime/Core/Public/Android/AndroidProperties.h"
 
 #if WITH_ENGINE
 	#include "Engine.h"
 	#include "TextureCompressorModule.h"
 #endif
-#include "TargetDeviceServices.h"
+
 #include "TargetPlatform.h"
 #include "AndroidDeviceDetection.h"
 
diff --git a/Engine/Source/Developer/Android/Android_PVRTCTargetPlatform/Android_PVRTCTargetPlatform.Build.cs b/Engine/Source/Developer/Android/Android_PVRTCTargetPlatform/Android_PVRTCTargetPlatform.Build.cs
--- a/Engine/Source/Developer/Android/Android_PVRTCTargetPlatform/Android_PVRTCTargetPlatform.Build.cs
+++ b/Engine/Source/Developer/Android/Android_PVRTCTargetPlatform/Android_PVRTCTargetPlatform.Build.cs
@@ -12,10 +12,7 @@ public class Android_PVRTCTargetPlatform : ModuleRules
 			new string[]
 			{
 				"Core",
-				"Sockets",
 				"TargetPlatform",
-				"Messaging",
-				"TargetDeviceServices",
 				"AndroidDeviceDetection",
 			}
 		);
diff --git a/Engine/Source/Developer/Android/Android_PVRTCTargetPlatform/Private/Android_PVRTCTargetPlatformPrivatePCH.h b/Engine/Source/Developer/Android/Android_PVRTCTargetPlatform/Private/Android_PVRTCTargetPlatformPrivatePCH.h
--- a/Engine/Source/Developer/Android/Android_PVRTCTargetPlatform/Private/Android_PVRTCTargetPlatformPrivatePCH.h
+++ b/Engine/Source/Developer/Android/Android_PVRTCTargetPlatform/Private/Android_PVRTCTargetPlatformPrivatePCH.h
@@ -13,14 +13,13 @@
 #include "Core.h"
 #include "ModuleInterface.h"
 #include "ModuleManager.h"
-#include "Messaging.h"
 #include "Runtime/Core/Public/Android/AndroidProperties.h"
 
 #if WITH_ENGINE
 	#include "Engine.h"
 	#include "TextureCompressorModule.h"
 #endif
-#include "TargetDeviceServices.h"
+
 #include "TargetPlatform.h"
 #include "AndroidDeviceDetection.h"
 
