Fix detection of model and device for some Android devices. We now also try and query the property information on the device if we can't get it from the adb devices list.

diff --git a/Engine/Source/Developer/Android/AndroidDeviceDetection/Private/AndroidDeviceDetectionModule.cpp b/Engine/Source/Developer/Android/AndroidDeviceDetection/Private/AndroidDeviceDetectionModule.cpp
--- a/Engine/Source/Developer/Android/AndroidDeviceDetection/Private/AndroidDeviceDetectionModule.cpp
+++ b/Engine/Source/Developer/Android/AndroidDeviceDetection/Private/AndroidDeviceDetectionModule.cpp
@@ -151,10 +151,26 @@ private:
 			// parse the device model
 			FString Model;
 			FParse::Value(*DeviceString, TEXT("model:"), Model);
+			if (Model.IsEmpty())
+			{
+				FString ModelCommand = FString::Printf(TEXT("-s %s shell getprop ro.product.model"), *SerialNumber);
+				FString RoProductModel;
+				ExecuteAdbCommand(*ModelCommand, &RoProductModel, nullptr);
+				const TCHAR* Ptr = *RoProductModel;
+				FParse::Line(&Ptr, Model);
+			}
 
 			// parse the device model
 			FString DeviceName;
 			FParse::Value(*DeviceString, TEXT("device:"), DeviceName);
+			if (DeviceName.IsEmpty())
+			{
+				FString DeviceCommand = FString::Printf(TEXT("-s %s shell getprop ro.product.device"), *SerialNumber);
+				FString RoProductDevice;
+				ExecuteAdbCommand(*DeviceCommand, &RoProductDevice, nullptr);
+				const TCHAR* Ptr = *RoProductDevice;
+				FParse::Line(&Ptr, DeviceName);
+			}
 
 			// add the device to the map
 			{
