Fixed build for android.

diff --git a/Engine/Source/Developer/Android/AndroidTargetPlatform/AndroidTargetPlatform.Build.cs b/Engine/Source/Developer/Android/AndroidTargetPlatform/AndroidTargetPlatform.Build.cs
--- a/Engine/Source/Developer/Android/AndroidTargetPlatform/AndroidTargetPlatform.Build.cs
+++ b/Engine/Source/Developer/Android/AndroidTargetPlatform/AndroidTargetPlatform.Build.cs
@@ -14,7 +14,6 @@ public class AndroidTargetPlatform : ModuleRules
 				"Core",
 				"CoreUObject",
 				"TargetPlatform",
-				"DesktopPlatform",
 				"AndroidDeviceDetection",
 			}
 		);
diff --git a/Engine/Source/Developer/Android/Android_ATCTargetPlatform/Android_ATCTargetPlatform.Build.cs b/Engine/Source/Developer/Android/Android_ATCTargetPlatform/Android_ATCTargetPlatform.Build.cs
--- a/Engine/Source/Developer/Android/Android_ATCTargetPlatform/Android_ATCTargetPlatform.Build.cs
+++ b/Engine/Source/Developer/Android/Android_ATCTargetPlatform/Android_ATCTargetPlatform.Build.cs
@@ -13,7 +13,6 @@ public class Android_ATCTargetPlatform : ModuleRules
 			{
 				"Core",
 				"TargetPlatform",
-				"DesktopPlatform",
 				"AndroidDeviceDetection",
 			}
 		);
diff --git a/Engine/Source/Developer/Android/Android_DXTTargetPlatform/Android_DXTTargetPlatform.Build.cs b/Engine/Source/Developer/Android/Android_DXTTargetPlatform/Android_DXTTargetPlatform.Build.cs
--- a/Engine/Source/Developer/Android/Android_DXTTargetPlatform/Android_DXTTargetPlatform.Build.cs
+++ b/Engine/Source/Developer/Android/Android_DXTTargetPlatform/Android_DXTTargetPlatform.Build.cs
@@ -13,7 +13,6 @@ public class Android_DXTTargetPlatform : ModuleRules
 			{
 				"Core",
 				"TargetPlatform",
-				"DesktopPlatform",
 				"AndroidDeviceDetection",
 			}
 		);
diff --git a/Engine/Source/Developer/Android/Android_ETC1TargetPlatform/Android_ETC1TargetPlatform.Build.cs b/Engine/Source/Developer/Android/Android_ETC1TargetPlatform/Android_ETC1TargetPlatform.Build.cs
--- a/Engine/Source/Developer/Android/Android_ETC1TargetPlatform/Android_ETC1TargetPlatform.Build.cs
+++ b/Engine/Source/Developer/Android/Android_ETC1TargetPlatform/Android_ETC1TargetPlatform.Build.cs
@@ -20,7 +20,6 @@ public class Android_ETC1TargetPlatform : ModuleRules
 			{
 				"Core",
 				"TargetPlatform",
-				"DesktopPlatform",
 				"AndroidDeviceDetection",
 			}
 		);
diff --git a/Engine/Source/Developer/Android/Android_ETC2TargetPlatform/Android_ETC2TargetPlatform.Build.cs b/Engine/Source/Developer/Android/Android_ETC2TargetPlatform/Android_ETC2TargetPlatform.Build.cs
--- a/Engine/Source/Developer/Android/Android_ETC2TargetPlatform/Android_ETC2TargetPlatform.Build.cs
+++ b/Engine/Source/Developer/Android/Android_ETC2TargetPlatform/Android_ETC2TargetPlatform.Build.cs
@@ -13,7 +13,6 @@ public class Android_ETC2TargetPlatform : ModuleRules
 			{
 				"Core",
 				"TargetPlatform",
-				"DesktopPlatform",
 				"AndroidDeviceDetection",
 			}
 		);
diff --git a/Engine/Source/Developer/Android/Android_PVRTCTargetPlatform/Android_PVRTCTargetPlatform.Build.cs b/Engine/Source/Developer/Android/Android_PVRTCTargetPlatform/Android_PVRTCTargetPlatform.Build.cs
--- a/Engine/Source/Developer/Android/Android_PVRTCTargetPlatform/Android_PVRTCTargetPlatform.Build.cs
+++ b/Engine/Source/Developer/Android/Android_PVRTCTargetPlatform/Android_PVRTCTargetPlatform.Build.cs
@@ -13,7 +13,6 @@ public class Android_PVRTCTargetPlatform : ModuleRules
 			{
 				"Core",
 				"TargetPlatform",
-				"DesktopPlatform",
 				"AndroidDeviceDetection",
 			}
 		);
