Added back in the DesktopPlatform to androidtargetPlatform.build.cs.

diff --git a/Engine/Source/Developer/Android/AndroidTargetPlatform/AndroidTargetPlatform.Build.cs b/Engine/Source/Developer/Android/AndroidTargetPlatform/AndroidTargetPlatform.Build.cs
--- a/Engine/Source/Developer/Android/AndroidTargetPlatform/AndroidTargetPlatform.Build.cs
+++ b/Engine/Source/Developer/Android/AndroidTargetPlatform/AndroidTargetPlatform.Build.cs
@@ -14,6 +14,7 @@ public class AndroidTargetPlatform : ModuleRules
 				"Core",
 				"CoreUObject",
 				"TargetPlatform",
+                "DesktopPlatform",
 				"AndroidDeviceDetection",
 			}
 		);
