Changed default Android depth buffer precision to 24 and added Android platform setting to choose preference.
#Android
#TTP 334728

diff --git a/Engine/Build/Android/Java/AndroidManifest.xml b/Engine/Build/Android/Java/AndroidManifest.xml
--- a/Engine/Build/Android/Java/AndroidManifest.xml
+++ b/Engine/Build/Android/Java/AndroidManifest.xml
@@ -21,6 +21,7 @@
 				<category android:name="android.intent.category.LAUNCHER" />
 			</intent-filter>
 		</activity>
+		<meta-data android:name="com.epicgames.ue4.GameActivity.DepthBufferPreference" android:value="0"/>
 		<meta-data android:name="com.google.android.gms.games.APP_ID"
 			android:value="@string/app_id" />
 		<meta-data android:name="com.google.android.gms.version"
diff --git a/Engine/Build/Android/Java/src/com/epicgames/ue4/GameActivity.java b/Engine/Build/Android/Java/src/com/epicgames/ue4/GameActivity.java
--- a/Engine/Build/Android/Java/src/com/epicgames/ue4/GameActivity.java
+++ b/Engine/Build/Android/Java/src/com/epicgames/ue4/GameActivity.java
@@ -17,6 +17,9 @@ import android.content.Intent;
 import android.content.res.AssetManager;
 import android.content.res.Configuration;
 import android.content.IntentSender.SendIntentException;
+import android.content.pm.ApplicationInfo;
+import android.content.pm.PackageManager;
+import android.content.pm.PackageManager.NameNotFoundException;
 
 import android.media.AudioManager;
 import android.util.DisplayMetrics;
@@ -238,9 +241,27 @@ public class GameActivity extends NativeActivity implements GoogleApiClient.Conn
 		// Grab a reference to the asset manager
 		AssetManagerReference = this.getAssets();
 
+		// Get the preferred depth buffer size from AndroidManifest.xml
+		int DepthBufferPreference = 0;
+		try {
+			ApplicationInfo ai = getPackageManager().getApplicationInfo(getPackageName(), PackageManager.GET_META_DATA);
+			Bundle bundle = ai.metaData;
+			if (bundle.containsKey("com.epicgames.ue4.GameActivity.DepthBufferPreference"))
+			{
+				DepthBufferPreference = bundle.getInt("com.epicgames.ue4.GameActivity.DepthBufferPreference");
+				Log.debug( "Found DepthBufferPreference = " + DepthBufferPreference);
+			} else {
+				Log.debug( "Did not find DepthBufferPreference, using default.");
+			}
+		} catch (NameNotFoundException e) {
+			Log.debug( "Failed to load meta-data: NameNotFound: " + e.getMessage());
+		} catch (NullPointerException e) {
+			Log.debug( "Failed to load meta-data: NullPointer: " + e.getMessage());
+		}
+
 		// tell the engine if this is a portrait app
 		nativeSetGlobalActivity();
-		nativeSetWindowInfo(getResources().getConfiguration().orientation == Configuration.ORIENTATION_PORTRAIT);
+		nativeSetWindowInfo(getResources().getConfiguration().orientation == Configuration.ORIENTATION_PORTRAIT, DepthBufferPreference);
 
 
 		Log.debug( "Android version is " + android.os.Build.VERSION.RELEASE );
@@ -818,7 +839,7 @@ public class GameActivity extends NativeActivity implements GoogleApiClient.Conn
 
 	public native boolean nativeIsShippingBuild();
 	public native void nativeSetGlobalActivity();
-	public native void nativeSetWindowInfo(boolean bIsPortrait);
+	public native void nativeSetWindowInfo(boolean bIsPortrait, int DepthBufferPreference);
 	public native void nativeSetObbInfo(String PackageName, int Version, int PatchVersion);
 	public native void nativeUpdateAchievements(JavaAchievement[] Achievements);
 	public native void nativeSetAndroidVersionInformation( String AndroidVersion, String PhoneMake, String PhoneModel );
diff --git a/Engine/Source/Developer/Android/AndroidPlatformEditor/Private/AndroidTargetSettingsCustomization.cpp b/Engine/Source/Developer/Android/AndroidPlatformEditor/Private/AndroidTargetSettingsCustomization.cpp
--- a/Engine/Source/Developer/Android/AndroidPlatformEditor/Private/AndroidTargetSettingsCustomization.cpp
+++ b/Engine/Source/Developer/Android/AndroidPlatformEditor/Private/AndroidTargetSettingsCustomization.cpp
@@ -115,6 +115,10 @@ void FAndroidTargetSettingsCustomization::BuildAppManifestSection(IDetailLayoutB
 	AppManifestCategory.AddProperty(OrientationProperty)
 		.EditCondition(SetupForPlatformAttribute, NULL);
 
+	TSharedRef<IPropertyHandle> DepthBufferPreferenceProperty = DetailLayout.GetProperty(GET_MEMBER_NAME_CHECKED(UAndroidRuntimeSettings, DepthBufferPreference));
+	DepthBufferPreferenceProperty->SetOnPropertyValueChanged(FSimpleDelegate::CreateRaw(this, &FAndroidTargetSettingsCustomization::OnDepthBufferPreferenceModified));
+	AppManifestCategory.AddProperty(DepthBufferPreferenceProperty)
+		.EditCondition(SetupForPlatformAttribute, NULL);
 
 	// Google Play category
 	IDetailCategoryBuilder& GooglePlayCategory = DetailLayout.EditCategory(TEXT("GooglePlayServices"));
@@ -289,6 +293,42 @@ FString FAndroidTargetSettingsCustomization::OrientationToString(const EAndroidS
 	return Enum->GetMetaData(TEXT("ManifestValue"), (int32)Orientation);
 }
 
+void FAndroidTargetSettingsCustomization::OnDepthBufferPreferenceModified()
+{
+	check(SetupForPlatformAttribute.Get() == true);
+
+
+	FManifestUpdateHelper Updater(GameManifestPath);
+
+	const FString PrefixTag(TEXT("meta-data android:name=\"com.epicgames.ue4.GameActivity.DepthBufferPreference\" android:value=\""));
+	const FString ClosingQuote(TEXT("\""));
+	const FString NewDepthBufferPreferenceString = DepthBufferPreferenceToString(GetDefault<UAndroidRuntimeSettings>()->DepthBufferPreference);
+
+	// First check if the key needs to be created
+	if (!Updater.HasKey(PrefixTag))
+	{
+		// Add the key for value replacement
+		const FString CloseActivityTag(TEXT("</activity"));
+		const FString CloseBracket(TEXT(">"));
+		const FString NewKeyString(TEXT(">\n\t\t<meta-data android:name=\"com.epicgames.ue4.GameActivity.DepthBufferPreference\" android:value=\"0\"/"));
+		Updater.ReplaceKey(CloseActivityTag, CloseBracket, NewKeyString);
+	}
+
+	Updater.ReplaceKey(PrefixTag, ClosingQuote, NewDepthBufferPreferenceString);
+
+	Updater.Finalize(GameManifestPath);
+}
+
+FString FAndroidTargetSettingsCustomization::DepthBufferPreferenceToString(const EAndroidDepthBufferPreference::Type DepthBufferPreference)
+{
+	extern ANDROIDRUNTIMESETTINGS_API class UEnum* Z_Construct_UEnum_UAndroidRuntimeSettings_EAndroidDepthBufferPreference();
+
+	UEnum* Enum = Z_Construct_UEnum_UAndroidRuntimeSettings_EAndroidDepthBufferPreference();
+	check(Enum != nullptr);
+
+	return Enum->GetMetaData(TEXT("ManifestValue"), (int32)DepthBufferPreference);
+}
+
 void FAndroidTargetSettingsCustomization::OnAppIDModified()
 {
 	check(SetupForPlatformAttribute.Get() == true);
diff --git a/Engine/Source/Developer/Android/AndroidPlatformEditor/Private/AndroidTargetSettingsCustomization.h b/Engine/Source/Developer/Android/AndroidPlatformEditor/Private/AndroidTargetSettingsCustomization.h
--- a/Engine/Source/Developer/Android/AndroidPlatformEditor/Private/AndroidTargetSettingsCustomization.h
+++ b/Engine/Source/Developer/Android/AndroidPlatformEditor/Private/AndroidTargetSettingsCustomization.h
@@ -37,6 +37,9 @@ private:
 	// Called when the orientation is modified
 	void OnOrientationModified();
 
+	// Called when the depth preference is modified
+	void OnDepthBufferPreferenceModified();
+
 	// Called when the app id is modified
 	void OnAppIDModified();
 
@@ -72,5 +75,8 @@ private:
 	// Converts an orientation enum to the associated string value
 	static FString OrientationToString(const EAndroidScreenOrientation::Type Orientation);
 
+	// Converts an depth preference enum to the associated string value
+	static FString DepthBufferPreferenceToString(const EAndroidDepthBufferPreference::Type DepthBufferPreference);
+
 	IDetailLayoutBuilder* SavedLayoutBuilder;
 };
diff --git a/Engine/Source/Developer/SharedSettingsWidgets/Private/ManifestUpdateHelper.cpp b/Engine/Source/Developer/SharedSettingsWidgets/Private/ManifestUpdateHelper.cpp
--- a/Engine/Source/Developer/SharedSettingsWidgets/Private/ManifestUpdateHelper.cpp
+++ b/Engine/Source/Developer/SharedSettingsWidgets/Private/ManifestUpdateHelper.cpp
@@ -48,6 +48,12 @@ void FManifestUpdateHelper::WriteError(FText NewError)
 	UE_LOG(LogInit, Warning, TEXT("Error during platform manifest modification: %s"), *NewError.ToString());
 }
 
+bool FManifestUpdateHelper::HasKey(const FString& MatchPrefix)
+{
+	const int32 PrefixPos = ManifestString.Find(MatchPrefix, ESearchCase::CaseSensitive);
+	return (PrefixPos != INDEX_NONE);
+}
+
 void FManifestUpdateHelper::ReplaceKey(const FString& MatchPrefix, const FString& MatchSuffix, const FString& NewInfix)
 {
 	if (ReplaceStringPortion(ManifestString, MatchPrefix, MatchSuffix, NewInfix))
diff --git a/Engine/Source/Developer/SharedSettingsWidgets/Public/ManifestUpdateHelper.h b/Engine/Source/Developer/SharedSettingsWidgets/Public/ManifestUpdateHelper.h
--- a/Engine/Source/Developer/SharedSettingsWidgets/Public/ManifestUpdateHelper.h
+++ b/Engine/Source/Developer/SharedSettingsWidgets/Public/ManifestUpdateHelper.h
@@ -16,6 +16,9 @@ public:
 	// Replaces the text in InOutString between MatchPrefix and MatchSuffix with NewInfix, returning true if it was found
 	static bool ReplaceStringPortion(FString& InOutString, const FString& MatchPrefix, const FString& MatchSuffix, const FString& NewInfix);
 
+	// Checks if a key exists matching MatchPrefix
+	bool HasKey(const FString& MatchPrefix);
+
 	// Replace a key in the manifest (between MatchPrefix and MatchSuffix) with NewInfix
 	void ReplaceKey(const FString& MatchPrefix, const FString& MatchSuffix, const FString& NewInfix);
 
diff --git a/Engine/Source/Runtime/Android/AndroidRuntimeSettings/Classes/AndroidRuntimeSettings.h b/Engine/Source/Runtime/Android/AndroidRuntimeSettings/Classes/AndroidRuntimeSettings.h
--- a/Engine/Source/Runtime/Android/AndroidRuntimeSettings/Classes/AndroidRuntimeSettings.h
+++ b/Engine/Source/Runtime/Android/AndroidRuntimeSettings/Classes/AndroidRuntimeSettings.h
@@ -35,6 +35,19 @@ namespace EAndroidScreenOrientation
 	};
 }
 
+// Depth buffer precision preferences
+UENUM()
+namespace EAndroidDepthBufferPreference
+{
+	enum Type
+	{
+		Default = 0 UMETA(DisplayName = "Default", ManifestValue = "0"),
+		Bits16 = 16 UMETA(DisplayName = "16-bit", ManifestValue = "16"),
+		Bits24 = 24 UMETA(DisplayName = "24-bit", ManifestValue = "24"),
+		Bits32 = 32 UMETA(DisplayName = "32-bit", ManifestValue = "32"),
+	};
+}
+
 /**
  * Holds the game-specific achievement name and corresponding ID from Google Play services.
  */
@@ -82,6 +95,10 @@ public:
 	UPROPERTY(GlobalConfig, EditAnywhere, Category = AppManifest)
 	TEnumAsByte<EAndroidScreenOrientation::Type> Orientation;
 
+	// The preferred depth buffer bitcount for Android
+	UPROPERTY(GlobalConfig, EditAnywhere, Category = AppManifest)
+	TEnumAsByte<EAndroidDepthBufferPreference::Type> DepthBufferPreference;
+
 	// Should Google Play support be enabled?
 	UPROPERTY(GlobalConfig, EditAnywhere, Category = GooglePlayServices)
 	bool bEnableGooglePlaySupport;
diff --git a/Engine/Source/Runtime/Core/Private/Android/AndroidWindow.cpp b/Engine/Source/Runtime/Core/Private/Android/AndroidWindow.cpp
--- a/Engine/Source/Runtime/Core/Private/Android/AndroidWindow.cpp
+++ b/Engine/Source/Runtime/Core/Private/Android/AndroidWindow.cpp
@@ -52,12 +52,19 @@ void FAndroidWindow::SetOSWindowHandle(void* InWindow)
 
 //This function is declared in the Java-defined class, GameActivity.java: "public native void nativeSetObbInfo(String PackageName, int Version, int PatchVersion);"
 static bool GAndroidIsPortrait = false;
-extern "C" void Java_com_epicgames_ue4_GameActivity_nativeSetWindowInfo(JNIEnv* jenv, jobject thiz, jboolean bIsPortrait)
+static int GAndroidDepthBufferPreference = 0;
+extern "C" void Java_com_epicgames_ue4_GameActivity_nativeSetWindowInfo(JNIEnv* jenv, jobject thiz, jboolean bIsPortrait, jint DepthBufferPreference)
 {
 	GAndroidIsPortrait = bIsPortrait == JNI_TRUE;
+	GAndroidDepthBufferPreference = DepthBufferPreference;
 	FPlatformMisc::LowLevelOutputDebugStringf(TEXT("App is running in %s\n"), GAndroidIsPortrait ? TEXT("Portrait") : TEXT("Landscape"));
 }
 
+int32 FAndroidWindow::GetDepthBufferPreference()
+{
+	return GAndroidDepthBufferPreference;
+}
+
 FPlatformRect FAndroidWindow::GetScreenRect()
 {
 	// currently hardcoding resolution
diff --git a/Engine/Source/Runtime/Core/Public/Android/AndroidWindow.h b/Engine/Source/Runtime/Core/Public/Android/AndroidWindow.h
--- a/Engine/Source/Runtime/Core/Public/Android/AndroidWindow.h
+++ b/Engine/Source/Runtime/Core/Public/Android/AndroidWindow.h
@@ -39,6 +39,7 @@ public:
 
 	static void CalculateSurfaceSize(void* InWindow, int32_t& SurfaceWidth, int32_t& SurfaceHeight);
 
+	static int32 GetDepthBufferPreference();
 
 protected:
 	/** @return true if the native window is currently in fullscreen mode, false otherwise */
diff --git a/Engine/Source/Runtime/OpenGLDrv/Private/Android/AndroidEGL.cpp b/Engine/Source/Runtime/OpenGLDrv/Private/Android/AndroidEGL.cpp
--- a/Engine/Source/Runtime/OpenGLDrv/Private/Android/AndroidEGL.cpp
+++ b/Engine/Source/Runtime/OpenGLDrv/Private/Android/AndroidEGL.cpp
@@ -109,7 +109,6 @@ EGLConfigParms::EGLConfigParms(const EGLConfigParms& Parms)
 	stencilSize = Parms.stencilSize;
 	sampleBuffers = Parms.sampleBuffers;
 	sampleSamples = Parms.sampleSamples;
-
 }
 
 EGLConfigParms::EGLConfigParms()  : 
@@ -118,11 +117,15 @@ validConfig (0)
 	,greenSize(8)
 	,blueSize(8)
 	,alphaSize(0)
-	,depthSize(16)
+	,depthSize(24)
 	,stencilSize(0)
 	,sampleBuffers(0)
 	,sampleSamples(0)
 {
+	// If not default, set the preference
+	int DepthBufferPreference = (int)FAndroidWindow::GetDepthBufferPreference();
+	if (DepthBufferPreference > 0)
+		depthSize = DepthBufferPreference;
 }
 
 AndroidEGL::AndroidEGL()
