multiple android target platform symbols CIS fix.

diff --git a/Engine/Source/Developer/Android/AndroidTargetPlatform/Private/AndroidTargetPlatform.inl b/Engine/Source/Developer/Android/AndroidTargetPlatform/Private/AndroidTargetPlatform.inl
--- a/Engine/Source/Developer/Android/AndroidTargetPlatform/Private/AndroidTargetPlatform.inl
+++ b/Engine/Source/Developer/Android/AndroidTargetPlatform/Private/AndroidTargetPlatform.inl
@@ -81,7 +81,7 @@ inline bool FAndroidTargetPlatform::IsRunningPlatform( ) const
 }
 
 
-bool FAndroidTargetPlatform::IsSdkInstalled(bool bProjectHasCode, FString& OutDocumentationPath) const
+inline bool FAndroidTargetPlatform::IsSdkInstalled(bool bProjectHasCode, FString& OutDocumentationPath) const
 {
 	OutDocumentationPath = FString("Shared/Tutorials/SettingUpAndroidTutorial");
 
@@ -132,7 +132,7 @@ bool FAndroidTargetPlatform::IsSdkInstalled(bool bProjectHasCode, FString& OutDo
 }
 
 
-bool FAndroidTargetPlatform::SupportsFeature( ETargetPlatformFeatures::Type Feature ) const
+inline bool FAndroidTargetPlatform::SupportsFeature( ETargetPlatformFeatures::Type Feature ) const
 {
 	if (Feature == ETargetPlatformFeatures::Packaging)
 	{
