Changed FAndoridTargetPlatform to take template parameter for the AndroidPlatformProperties, so that the PlatformProperties retrieved from the FTargetplatform are correct.

diff --git a/Engine/Source/Developer/Android/AndroidTargetPlatform/Private/AndroidTargetPlatform.h b/Engine/Source/Developer/Android/AndroidTargetPlatform/Private/AndroidTargetPlatform.h
--- a/Engine/Source/Developer/Android/AndroidTargetPlatform/Private/AndroidTargetPlatform.h
+++ b/Engine/Source/Developer/Android/AndroidTargetPlatform/Private/AndroidTargetPlatform.h
@@ -45,8 +45,9 @@ namespace AndroidTexFormat
 /**
  * FAndroidTargetPlatform, abstraction for cooking Android platforms
  */
+template<class TPlatformProperties>
 class FAndroidTargetPlatform
-	: public TTargetPlatformBase< FAndroidPlatformProperties >
+	: public TTargetPlatformBase< TPlatformProperties >
 {
 public:
 
diff --git a/Engine/Source/Developer/Android/AndroidTargetPlatform/Private/AndroidTargetPlatform.inl b/Engine/Source/Developer/Android/AndroidTargetPlatform/Private/AndroidTargetPlatform.inl
--- a/Engine/Source/Developer/Android/AndroidTargetPlatform/Private/AndroidTargetPlatform.inl
+++ b/Engine/Source/Developer/Android/AndroidTargetPlatform/Private/AndroidTargetPlatform.inl
@@ -8,7 +8,8 @@
 /* FAndroidTargetPlatform structors
  *****************************************************************************/
 
-inline FAndroidTargetPlatform::FAndroidTargetPlatform( ) :
+template<class TPlatformProperties>
+inline FAndroidTargetPlatform<TPlatformProperties>::FAndroidTargetPlatform( ) :
 	DeviceDetection(nullptr)
 {
 	#if WITH_ENGINE
@@ -22,7 +23,8 @@ inline FAndroidTargetPlatform::FAndroidTargetPlatform( ) :
 }
 
 
-inline FAndroidTargetPlatform::~FAndroidTargetPlatform()
+template<class TPlatformProperties>
+inline FAndroidTargetPlatform<TPlatformProperties>::~FAndroidTargetPlatform()
 { 
 	 FTicker::GetCoreTicker().RemoveTicker(TickDelegate);
 }
@@ -31,7 +33,8 @@ inline FAndroidTargetPlatform::~FAndroidTargetPlatform()
 /* ITargetPlatform overrides
  *****************************************************************************/
 
-inline void FAndroidTargetPlatform::GetAllDevices( TArray<ITargetDevicePtr>& OutDevices ) const
+template<class TPlatformProperties>
+inline void FAndroidTargetPlatform<TPlatformProperties>::GetAllDevices( TArray<ITargetDevicePtr>& OutDevices ) const
 {
 	OutDevices.Reset();
 
@@ -41,14 +44,14 @@ inline void FAndroidTargetPlatform::GetAllDevices( TArray<ITargetDevicePtr>& Out
 	}
 }
 
-
-inline ECompressionFlags FAndroidTargetPlatform::GetBaseCompressionMethod( ) const
+template<class TPlatformProperties>
+inline ECompressionFlags FAndroidTargetPlatform<TPlatformProperties>::GetBaseCompressionMethod( ) const
 {
 	return COMPRESS_ZLIB;
 }
 
-
-inline ITargetDevicePtr FAndroidTargetPlatform::GetDefaultDevice( ) const
+template<class TPlatformProperties>
+inline ITargetDevicePtr FAndroidTargetPlatform<TPlatformProperties>::GetDefaultDevice( ) const
 {
 	// return the first device in the list
 	if (Devices.Num() > 0)
@@ -63,8 +66,8 @@ inline ITargetDevicePtr FAndroidTargetPlatform::GetDefaultDevice( ) const
 	return nullptr;
 }
 
-
-inline ITargetDevicePtr FAndroidTargetPlatform::GetDevice( const FTargetDeviceId& DeviceId )
+template<class TPlatformProperties>
+inline ITargetDevicePtr FAndroidTargetPlatform<TPlatformProperties>::GetDevice( const FTargetDeviceId& DeviceId )
 {
 	if (DeviceId.GetPlatformName() == PlatformName())
 	{
@@ -74,14 +77,15 @@ inline ITargetDevicePtr FAndroidTargetPlatform::GetDevice( const FTargetDeviceId
 	return nullptr;
 }
 
-
-inline bool FAndroidTargetPlatform::IsRunningPlatform( ) const
+template<class TPlatformProperties>
+inline bool FAndroidTargetPlatform<TPlatformProperties>::IsRunningPlatform( ) const
 {
 	return false; // This platform never runs the target platform framework
 }
 
 
-inline bool FAndroidTargetPlatform::IsSdkInstalled(bool bProjectHasCode, FString& OutDocumentationPath) const
+template<class TPlatformProperties>
+inline bool FAndroidTargetPlatform<TPlatformProperties>::IsSdkInstalled(bool bProjectHasCode, FString& OutDocumentationPath) const
 {
 	OutDocumentationPath = FString("Shared/Tutorials/SettingUpAndroidTutorial");
 
@@ -132,39 +136,43 @@ inline bool FAndroidTargetPlatform::IsSdkInstalled(bool bProjectHasCode, FString
 }
 
 
-inline bool FAndroidTargetPlatform::SupportsFeature( ETargetPlatformFeatures::Type Feature ) const
+template<class TPlatformProperties>
+inline bool FAndroidTargetPlatform<TPlatformProperties>::SupportsFeature( ETargetPlatformFeatures::Type Feature ) const
 {
 	if (Feature == ETargetPlatformFeatures::Packaging)
 	{
 		return true;
 	}
 
-	return TTargetPlatformBase<FAndroidPlatformProperties>::SupportsFeature(Feature);
+	return TTargetPlatformBase<TPlatformProperties>::SupportsFeature(Feature);
 }
 
 
 #if WITH_ENGINE
 
-inline void FAndroidTargetPlatform::GetAllPossibleShaderFormats( TArray<FName>& OutFormats ) const
+template<class TPlatformProperties>
+inline void FAndroidTargetPlatform<TPlatformProperties>::GetAllPossibleShaderFormats( TArray<FName>& OutFormats ) const
 {
 	static FName NAME_OPENGL_ES2(TEXT("GLSL_ES2"));
 	OutFormats.AddUnique(NAME_OPENGL_ES2);
 }
 
-
-inline void FAndroidTargetPlatform::GetAllTargetedShaderFormats( TArray<FName>& OutFormats ) const
+template<class TPlatformProperties>
+inline void FAndroidTargetPlatform<TPlatformProperties>::GetAllTargetedShaderFormats( TArray<FName>& OutFormats ) const
 {
 	GetAllPossibleShaderFormats(OutFormats);
 }
 
 
-inline const FStaticMeshLODSettings& FAndroidTargetPlatform::GetStaticMeshLODSettings( ) const
+template<class TPlatformProperties>
+inline const FStaticMeshLODSettings& FAndroidTargetPlatform<TPlatformProperties>::GetStaticMeshLODSettings( ) const
 {
 	return StaticMeshLODSettings;
 }
 
 
-inline void FAndroidTargetPlatform::GetTextureFormats( const UTexture* InTexture, TArray<FName>& OutFormats ) const
+template<class TPlatformProperties>
+inline void FAndroidTargetPlatform<TPlatformProperties>::GetTextureFormats( const UTexture* InTexture, TArray<FName>& OutFormats ) const
 {
 	check(InTexture);
 
@@ -253,13 +261,15 @@ inline void FAndroidTargetPlatform::GetTextureFormats( const UTexture* InTexture
 }
 
 
-inline const struct FTextureLODSettings& FAndroidTargetPlatform::GetTextureLODSettings( ) const
+template<class TPlatformProperties>
+const FTextureLODSettings& FAndroidTargetPlatform<TPlatformProperties>::GetTextureLODSettings( ) const
 {
 	return TextureLODSettings;
 }
 
 
-inline FName FAndroidTargetPlatform::GetWaveFormat( class USoundWave* Wave ) const
+template<class TPlatformProperties>
+FName FAndroidTargetPlatform<TPlatformProperties>::GetWaveFormat( class USoundWave* Wave ) const
 {
 	static FName NAME_OGG(TEXT("OGG"));		//@todo android: probably not ogg
 
@@ -272,7 +282,8 @@ inline FName FAndroidTargetPlatform::GetWaveFormat( class USoundWave* Wave ) con
 /* FAndroidTargetPlatform implementation
  *****************************************************************************/
 
-inline void FAndroidTargetPlatform::AddTextureFormatIfSupports( FName Format, TArray<FName>& OutFormats ) const
+template<class TPlatformProperties>
+inline void FAndroidTargetPlatform<TPlatformProperties>::AddTextureFormatIfSupports( FName Format, TArray<FName>& OutFormats ) const
 {
 	if (SupportsTextureFormat(Format))
 	{
@@ -284,7 +295,8 @@ inline void FAndroidTargetPlatform::AddTextureFormatIfSupports( FName Format, TA
 /* FAndroidTargetPlatform callbacks
  *****************************************************************************/
 
-inline bool FAndroidTargetPlatform::HandleTicker( float DeltaTime )
+template<class TPlatformProperties>
+inline bool FAndroidTargetPlatform<TPlatformProperties>::HandleTicker( float DeltaTime )
 {
 	if (DeviceDetection == nullptr)
 	{
diff --git a/Engine/Source/Developer/Android/AndroidTargetPlatform/Private/AndroidTargetPlatformModule.cpp b/Engine/Source/Developer/Android/AndroidTargetPlatform/Private/AndroidTargetPlatformModule.cpp
--- a/Engine/Source/Developer/Android/AndroidTargetPlatform/Private/AndroidTargetPlatformModule.cpp
+++ b/Engine/Source/Developer/Android/AndroidTargetPlatform/Private/AndroidTargetPlatformModule.cpp
@@ -33,7 +33,7 @@ public:
 	{
 		if (AndroidTargetSingleton == NULL)
 		{
-			AndroidTargetSingleton = new FAndroidTargetPlatform();	
+			AndroidTargetSingleton = new FAndroidTargetPlatform<FAndroidPlatformProperties>();	
 		}
 		
 		return AndroidTargetSingleton;
diff --git a/Engine/Source/Developer/Android/Android_ATCTargetPlatform/Private/Android_ATCTargetPlatformModule.cpp b/Engine/Source/Developer/Android/Android_ATCTargetPlatform/Private/Android_ATCTargetPlatformModule.cpp
--- a/Engine/Source/Developer/Android/Android_ATCTargetPlatform/Private/Android_ATCTargetPlatformModule.cpp
+++ b/Engine/Source/Developer/Android/Android_ATCTargetPlatform/Private/Android_ATCTargetPlatformModule.cpp
@@ -13,7 +13,7 @@
  * Android cooking platform which cooks only ATC based textures.
  */
 class FAndroid_ATCTargetPlatform
-	: public FAndroidTargetPlatform
+	: public FAndroidTargetPlatform<FAndroid_ATCPlatformProperties>
 {
 	virtual FString GetAndroidVariantName( ) override
 	{
@@ -27,7 +27,7 @@ class FAndroid_ATCTargetPlatform
 
 	virtual FString PlatformName() const override
 	{
-		return FString(FAndroid_ATCCPlatformProperties::PlatformName());
+		return FString(FAndroid_ATCPlatformProperties::PlatformName());
 	}
 
 	virtual bool SupportsTextureFormat( FName Format ) const override
diff --git a/Engine/Source/Developer/Android/Android_DXTTargetPlatform/Private/Android_DXTTargetPlatformModule.cpp b/Engine/Source/Developer/Android/Android_DXTTargetPlatform/Private/Android_DXTTargetPlatformModule.cpp
--- a/Engine/Source/Developer/Android/Android_DXTTargetPlatform/Private/Android_DXTTargetPlatformModule.cpp
+++ b/Engine/Source/Developer/Android/Android_DXTTargetPlatform/Private/Android_DXTTargetPlatformModule.cpp
@@ -13,7 +13,7 @@
  * Android cooking platform which cooks only DXT based textures.
  */
 class FAndroid_DXTTargetPlatform
-	: public FAndroidTargetPlatform
+	: public FAndroidTargetPlatform<FAndroid_DXTPlatformProperties>
 {
 	virtual FString GetAndroidVariantName( ) override
 	{
diff --git a/Engine/Source/Developer/Android/Android_ETC1TargetPlatform/Private/Android_ETC1TargetPlatformModule.cpp b/Engine/Source/Developer/Android/Android_ETC1TargetPlatform/Private/Android_ETC1TargetPlatformModule.cpp
--- a/Engine/Source/Developer/Android/Android_ETC1TargetPlatform/Private/Android_ETC1TargetPlatformModule.cpp
+++ b/Engine/Source/Developer/Android/Android_ETC1TargetPlatform/Private/Android_ETC1TargetPlatformModule.cpp
@@ -13,7 +13,7 @@
  * Android cooking platform which cooks only ETC1 based textures.
  */
 class FAndroid_ETC1TargetPlatform
-	: public FAndroidTargetPlatform
+	: public FAndroidTargetPlatform<FAndroid_ETC1PlatformProperties>
 {
 public:
 
diff --git a/Engine/Source/Developer/Android/Android_ETC2TargetPlatform/Private/Android_ETC2TargetPlatformModule.cpp b/Engine/Source/Developer/Android/Android_ETC2TargetPlatform/Private/Android_ETC2TargetPlatformModule.cpp
--- a/Engine/Source/Developer/Android/Android_ETC2TargetPlatform/Private/Android_ETC2TargetPlatformModule.cpp
+++ b/Engine/Source/Developer/Android/Android_ETC2TargetPlatform/Private/Android_ETC2TargetPlatformModule.cpp
@@ -13,7 +13,7 @@
  * Android cooking platform which cooks only ETC2 based textures.
  */
 class FAndroid_ETC2TargetPlatform
-	: public FAndroidTargetPlatform
+	: public FAndroidTargetPlatform<FAndroid_ETC2PlatformProperties>
 {
 public:
 
diff --git a/Engine/Source/Developer/Android/Android_PVRTCTargetPlatform/Private/Android_PVRTCTargetPlatformModule.cpp b/Engine/Source/Developer/Android/Android_PVRTCTargetPlatform/Private/Android_PVRTCTargetPlatformModule.cpp
--- a/Engine/Source/Developer/Android/Android_PVRTCTargetPlatform/Private/Android_PVRTCTargetPlatformModule.cpp
+++ b/Engine/Source/Developer/Android/Android_PVRTCTargetPlatform/Private/Android_PVRTCTargetPlatformModule.cpp
@@ -12,7 +12,7 @@
  * Android cooking platform which cooks only PVRTC based textures.
  */
 class FAndroid_PVRTCTargetPlatform
-	: public FAndroidTargetPlatform
+	: public FAndroidTargetPlatform<FAndroid_PVRTCPlatformProperties>
 {
 	virtual FString GetAndroidVariantName( ) override
 	{
diff --git a/Engine/Source/Runtime/Core/Public/Android/AndroidProperties.h b/Engine/Source/Runtime/Core/Public/Android/AndroidProperties.h
--- a/Engine/Source/Runtime/Core/Public/Android/AndroidProperties.h
+++ b/Engine/Source/Runtime/Core/Public/Android/AndroidProperties.h
@@ -107,7 +107,7 @@ struct FAndroid_PVRTCPlatformProperties : public FAndroidPlatformProperties
 	}
 };
 
-struct FAndroid_ATCCPlatformProperties : public FAndroidPlatformProperties
+struct FAndroid_ATCPlatformProperties : public FAndroidPlatformProperties
 {
 	static FORCEINLINE const char* PlatformName()
 	{
