Fixed mac build.

diff --git a/Engine/Source/Developer/Android/AndroidTargetPlatform/Private/AndroidTargetPlatform.inl b/Engine/Source/Developer/Android/AndroidTargetPlatform/Private/AndroidTargetPlatform.inl
--- a/Engine/Source/Developer/Android/AndroidTargetPlatform/Private/AndroidTargetPlatform.inl
+++ b/Engine/Source/Developer/Android/AndroidTargetPlatform/Private/AndroidTargetPlatform.inl
@@ -13,7 +13,7 @@ inline FAndroidTargetPlatform<TPlatformProperties>::FAndroidTargetPlatform( ) :
 	DeviceDetection(nullptr)
 {
 	#if WITH_ENGINE
-		FConfigCacheIni::LoadLocalIniFile(EngineSettings, TEXT("Engine"), true, *PlatformName());
+		FConfigCacheIni::LoadLocalIniFile(EngineSettings, TEXT("Engine"), true, *TTargetPlatformBase<TPlatformProperties>::PlatformName());
 		TextureLODSettings.Initialize(EngineSettings, TEXT("SystemSettings"));
 		StaticMeshLODSettings.Initialize(EngineSettings);
 	#endif
@@ -69,7 +69,7 @@ inline ITargetDevicePtr FAndroidTargetPlatform<TPlatformProperties>::GetDefaultD
 template<class TPlatformProperties>
 inline ITargetDevicePtr FAndroidTargetPlatform<TPlatformProperties>::GetDevice( const FTargetDeviceId& DeviceId )
 {
-	if (DeviceId.GetPlatformName() == PlatformName())
+	if (DeviceId.GetPlatformName() == TTargetPlatformBase<TPlatformProperties>::PlatformName())
 	{
 		return Devices.FindRef(DeviceId.GetDeviceName());
 	}
