Changed TEXT("") to FName when checking the audio type
Changed the routine which decides which format so that it only does the compare once and then returns that value for the whole cook.

diff --git a/Engine/Source/Developer/Android/AndroidTargetPlatform/Private/AndroidTargetPlatform.inl b/Engine/Source/Developer/Android/AndroidTargetPlatform/Private/AndroidTargetPlatform.inl
--- a/Engine/Source/Developer/Android/AndroidTargetPlatform/Private/AndroidTargetPlatform.inl
+++ b/Engine/Source/Developer/Android/AndroidTargetPlatform/Private/AndroidTargetPlatform.inl
@@ -275,28 +275,41 @@ const FTextureLODSettings& FAndroidTargetPlatform<TPlatformProperties>::GetTextu
 template<class TPlatformProperties>
 FName FAndroidTargetPlatform<TPlatformProperties>::GetWaveFormat( class USoundWave* Wave ) const
 {
-	FString audioSetting;
-	if (!GConfig->GetString(TEXT("/Script/UnrealEd.CookerSettings"), TEXT("AndroidAudio"), audioSetting, GEngineIni))
+	static bool formatRead = false;
+	static FName NAME_FORMAT;
+
+	if (!formatRead)
 	{
-		audioSetting = TEXT("DEFAULT");
-	}
+		formatRead = true;
+
+		FString audioSetting;
+		if (!GConfig->GetString(TEXT("/Script/UnrealEd.CookerSettings"), TEXT("AndroidAudio"), audioSetting, GEngineIni))
+		{
+			audioSetting = TEXT("DEFAULT");
+		}
 
 #if WITH_OGGVORBIS
-	if (audioSetting == TEXT("OGG") || audioSetting == TEXT("Default"))
-	{
-		static FName NAME_OGG(TEXT("OGG"));
-		return NAME_OGG;
-	}
+		if (audioSetting == TEXT("OGG") || audioSetting == TEXT("Default"))
+		{
+			static FName NAME_OGG(TEXT("OGG"));
+			NAME_FORMAT = NAME_OGG;
+		}
 #else
-	if (audioSetting == TEXT("OGG"))
-	{
-		UE_LOG(LogAudio, Error, TEXT("Attemped to select Ogg Vorbis encoding when the cooker is built without Ogg Vorbis support."));
-	}
+		if (audioSetting == TEXT("OGG"))
+		{
+			UE_LOG(LogAudio, Error, TEXT("Attemped to select Ogg Vorbis encoding when the cooker is built without Ogg Vorbis support."));
+		}
 #endif
+		else
+		{
 	
-	// Otherwise return ADPCM as it'll either be option '2' or 'default' depending on WITH_OGGVORBIS config
-	static FName NAME_ADPCM(TEXT("ADPCM"));
-	return NAME_ADPCM;
+			// Otherwise return ADPCM as it'll either be option '2' or 'default' depending on WITH_OGGVORBIS config
+			static FName NAME_ADPCM(TEXT("ADPCM"));
+			NAME_FORMAT = NAME_ADPCM;
+		}
+	}
+	UE_LOG(LogAudio, Warning, TEXT("Audio format is %s"), *NAME_FORMAT.ToString());
+	return NAME_FORMAT;
 }
 
 #endif //WITH_ENGINE
diff --git a/Engine/Source/Runtime/Android/AndroidAudio/Private/AndroidAudioDevice.cpp b/Engine/Source/Runtime/Android/AndroidAudio/Private/AndroidAudioDevice.cpp
--- a/Engine/Source/Runtime/Android/AndroidAudio/Private/AndroidAudioDevice.cpp
+++ b/Engine/Source/Runtime/Android/AndroidAudio/Private/AndroidAudioDevice.cpp
@@ -158,13 +158,14 @@ FSoundSource* FSLESAudioDevice::CreateSoundSource()
 bool FSLESAudioDevice::HasCompressedAudioInfoClass(USoundWave* SoundWave)
 {
 #if WITH_OGGVORBIS
-	if (SoundWave->CompressionName.IsValid() && SoundWave->CompressionName == TEXT("OGG"))
+	static FName NAME_OGG(TEXT("OGG"));
+	if (SoundWave->CompressionName.IsValid() && SoundWave->CompressionName == NAME_OGG)
 	{
 		return true;
 	}
 #endif
-
-	if (SoundWave->CompressionName.IsValid() && SoundWave->CompressionName == TEXT("ADPCM"))
+	static FName NAME_ADPCM(TEXT("ADPCM"));
+	if (SoundWave->CompressionName.IsValid() && SoundWave->CompressionName == NAME_ADPCM)
 	{
 		return true;
 	}
@@ -176,13 +177,14 @@ bool FSLESAudioDevice::HasCompressedAudioInfoClass(USoundWave* SoundWave)
 class ICompressedAudioInfo* FSLESAudioDevice::CreateCompressedAudioInfo(USoundWave* SoundWave)
 {
 #if WITH_OGGVORBIS
-	if (SoundWave->CompressionName.IsValid() && SoundWave->CompressionName == TEXT("OGG"))
+	static FName NAME_OGG(TEXT("OGG"));
+	if (SoundWave->CompressionName.IsValid() && SoundWave->CompressionName == NAME_OGG)
 	{
 		return new FVorbisAudioInfo();
 	}
 #endif
-
-	if (SoundWave->CompressionName.IsValid() && SoundWave->CompressionName == TEXT("ADPCM"))
+	static FName NAME_ADPCM(TEXT("ADPCM"));
+	if (SoundWave->CompressionName.IsValid() && SoundWave->CompressionName == NAME_ADPCM)
 	{
 		return new FADPCMAudioInfo();
 	}
