Removed unneed log spam from Audio Compression Type detection code.

diff --git a/Engine/Source/Developer/Android/AndroidTargetPlatform/Private/AndroidTargetPlatform.inl b/Engine/Source/Developer/Android/AndroidTargetPlatform/Private/AndroidTargetPlatform.inl
--- a/Engine/Source/Developer/Android/AndroidTargetPlatform/Private/AndroidTargetPlatform.inl
+++ b/Engine/Source/Developer/Android/AndroidTargetPlatform/Private/AndroidTargetPlatform.inl
@@ -308,7 +308,6 @@ FName FAndroidTargetPlatform<TPlatformProperties>::GetWaveFormat( class USoundWa
 			NAME_FORMAT = NAME_ADPCM;
 		}
 	}
-	UE_LOG(LogAudio, Warning, TEXT("Audio format is %s"), *NAME_FORMAT.ToString());
 	return NAME_FORMAT;
 }
 
