Android: Improving device detection/display
- Improve device detection logic to detect unauthorized devices and query SDK and release versions
- Add SDK versions to OS version display in the device manager

diff --git a/Engine/Source/Developer/Android/AndroidDeviceDetection/Private/AndroidDeviceDetectionModule.cpp b/Engine/Source/Developer/Android/AndroidDeviceDetection/Private/AndroidDeviceDetectionModule.cpp
--- a/Engine/Source/Developer/Android/AndroidDeviceDetection/Private/AndroidDeviceDetectionModule.cpp
+++ b/Engine/Source/Developer/Android/AndroidDeviceDetection/Private/AndroidDeviceDetectionModule.cpp
@@ -117,71 +117,95 @@ private:
 				continue;
 			}
 
-			FString SerialNumber = DeviceString.Left(TabIndex);
+			FAndroidDeviceInfo NewDeviceInfo;
 
-			// add it to our list of currently connected devices
-			CurrentlyConnectedDevices.Add(SerialNumber);
-
-			// move on to next device if this one is already a known device
-			if (DeviceMap.Contains(SerialNumber))
-			{
-				continue;
-			}
-
-			// get the GL extensions string (and a bunch of other stuff)
-			FString GLESExtensions;
+			NewDeviceInfo.SerialNumber = DeviceString.Left(TabIndex);
+			const FString DeviceState = DeviceString.Mid(TabIndex + 1).Trim();
 
-			FString ExtensionsCommand = FString::Printf(TEXT("-s %s shell dumpsys SurfaceFlinger"), *SerialNumber);
-			if (!ExecuteAdbCommand(*ExtensionsCommand, &GLESExtensions, nullptr))
-			{
-				continue;
-			}
+			NewDeviceInfo.bUnauthorizedDevice = DeviceState == TEXT("unauthorized");
 
-			// grab the GL ES version
-			FString GLESVersionString;
+			// add it to our list of currently connected devices
+			CurrentlyConnectedDevices.Add(NewDeviceInfo.SerialNumber);
 
-			FString VersionCommand = FString::Printf(TEXT("-s %s shell getprop ro.opengles.version"), *SerialNumber);
-			if (!ExecuteAdbCommand(*VersionCommand, &GLESVersionString, nullptr))
+			// move on to next device if this one is already a known device
+			if (DeviceMap.Contains(NewDeviceInfo.SerialNumber))
 			{
 				continue;
 			}
 
-			int GLESVersion = FCString::Atoi(*GLESVersionString);
-
-			// parse the device model
-			FString Model;
-			FParse::Value(*DeviceString, TEXT("model:"), Model);
-			if (Model.IsEmpty())
+			if (NewDeviceInfo.bUnauthorizedDevice)
 			{
-				FString ModelCommand = FString::Printf(TEXT("-s %s shell getprop ro.product.model"), *SerialNumber);
-				FString RoProductModel;
-				ExecuteAdbCommand(*ModelCommand, &RoProductModel, nullptr);
-				const TCHAR* Ptr = *RoProductModel;
-				FParse::Line(&Ptr, Model);
+				NewDeviceInfo.DeviceName = TEXT("Unauthorized - enable USB debugging");
 			}
-
-			// parse the device model
-			FString DeviceName;
-			FParse::Value(*DeviceString, TEXT("device:"), DeviceName);
-			if (DeviceName.IsEmpty())
+			else
 			{
-				FString DeviceCommand = FString::Printf(TEXT("-s %s shell getprop ro.product.device"), *SerialNumber);
-				FString RoProductDevice;
-				ExecuteAdbCommand(*DeviceCommand, &RoProductDevice, nullptr);
-				const TCHAR* Ptr = *RoProductDevice;
-				FParse::Line(&Ptr, DeviceName);
+				// grab the Android version
+				const FString AndroidVersionCommand = FString::Printf(TEXT("-s %s shell getprop ro.build.version.release"), *NewDeviceInfo.SerialNumber);
+				if (!ExecuteAdbCommand(*AndroidVersionCommand, &NewDeviceInfo.HumanAndroidVersion, nullptr))
+				{
+					continue;
+				}
+				NewDeviceInfo.HumanAndroidVersion = NewDeviceInfo.HumanAndroidVersion.Replace(TEXT("\r"), TEXT("")).Replace(TEXT("\n"), TEXT(""));
+				NewDeviceInfo.HumanAndroidVersion.Trim().TrimTrailing();
+
+				// grab the Android SDK version
+				const FString SDKVersionCommand = FString::Printf(TEXT("-s %s shell getprop ro.build.version.sdk"), *NewDeviceInfo.SerialNumber);
+				FString SDKVersionString;
+				if (!ExecuteAdbCommand(*SDKVersionCommand, &SDKVersionString, nullptr))
+				{
+					continue;
+				}
+				NewDeviceInfo.SDKVersion = FCString::Atoi(*SDKVersionString);
+				if (NewDeviceInfo.SDKVersion <= 0)
+				{
+					NewDeviceInfo.SDKVersion = INDEX_NONE;
+				}
+
+				// get the GL extensions string (and a bunch of other stuff)
+				const FString ExtensionsCommand = FString::Printf(TEXT("-s %s shell dumpsys SurfaceFlinger"), *NewDeviceInfo.SerialNumber);
+				if (!ExecuteAdbCommand(*ExtensionsCommand, &NewDeviceInfo.GLESExtensions, nullptr))
+				{
+					continue;
+				}
+
+				// grab the GL ES version
+				FString GLESVersionString;
+				const FString GLVersionCommand = FString::Printf(TEXT("-s %s shell getprop ro.opengles.version"), *NewDeviceInfo.SerialNumber);
+				if (!ExecuteAdbCommand(*GLVersionCommand, &GLESVersionString, nullptr))
+				{
+					continue;
+				}
+				NewDeviceInfo.GLESVersion = FCString::Atoi(*GLESVersionString);
+
+				// parse the device model
+				FParse::Value(*DeviceString, TEXT("model:"), NewDeviceInfo.Model);
+				if (NewDeviceInfo.Model.IsEmpty())
+				{
+					FString ModelCommand = FString::Printf(TEXT("-s %s shell getprop ro.product.model"), *NewDeviceInfo.SerialNumber);
+					FString RoProductModel;
+					ExecuteAdbCommand(*ModelCommand, &RoProductModel, nullptr);
+					const TCHAR* Ptr = *RoProductModel;
+					FParse::Line(&Ptr, NewDeviceInfo.Model);
+				}
+
+				// parse the device name
+				FParse::Value(*DeviceString, TEXT("device:"), NewDeviceInfo.DeviceName);
+				if (NewDeviceInfo.DeviceName.IsEmpty())
+				{
+					FString DeviceCommand = FString::Printf(TEXT("-s %s shell getprop ro.product.device"), *NewDeviceInfo.SerialNumber);
+					FString RoProductDevice;
+					ExecuteAdbCommand(*DeviceCommand, &RoProductDevice, nullptr);
+					const TCHAR* Ptr = *RoProductDevice;
+					FParse::Line(&Ptr, NewDeviceInfo.DeviceName);
+				}
 			}
 
 			// add the device to the map
 			{
 				FScopeLock ScopeLock(DeviceMapLock);
 
-				FAndroidDeviceInfo& DeviceInfo = DeviceMap.Add(SerialNumber);
-				DeviceInfo.SerialNumber = SerialNumber;
-				DeviceInfo.Model = Model;
-				DeviceInfo.DeviceName = DeviceName;
-				DeviceInfo.GLESExtensions = GLESExtensions;
-				DeviceInfo.GLESVersion = GLESVersion;
+				FAndroidDeviceInfo& SavedDeviceInfo = DeviceMap.Add(NewDeviceInfo.SerialNumber);
+				SavedDeviceInfo = NewDeviceInfo;
 			}
 		}
 
diff --git a/Engine/Source/Developer/Android/AndroidDeviceDetection/Public/Interfaces/IAndroidDeviceDetection.h b/Engine/Source/Developer/Android/AndroidDeviceDetection/Public/Interfaces/IAndroidDeviceDetection.h
--- a/Engine/Source/Developer/Android/AndroidDeviceDetection/Public/Interfaces/IAndroidDeviceDetection.h
+++ b/Engine/Source/Developer/Android/AndroidDeviceDetection/Public/Interfaces/IAndroidDeviceDetection.h
@@ -9,11 +9,36 @@
 
 struct FAndroidDeviceInfo
 {
+	// Device serial number, used to route ADB commands to a specific device
 	FString SerialNumber;
+
+	// Device model name
 	FString Model;
+
+	// Device name
 	FString DeviceName;
-	FString GLESExtensions;		//@todo android: currently retrieved from SurfaceFlinger running 1.1, but so far seems ok for supported ES2 compression formats
-	int GLESVersion;
+
+	// User-visible version of android installed (ro.build.version.release)
+	FString HumanAndroidVersion;
+
+	// Android SDK version supported by the device (ro.build.version.sdk - note: deprecated in 4 according to docs, but version 4 devices return an empty string when querying the 'replacement' SDK_INT)
+	int32 SDKVersion;
+
+	// List of supported OpenGL extensions (retrieved via SurfaceFlinger)
+	FString GLESExtensions;
+
+	// Supported GLES version (ro.opengles.version)
+	int32 GLESVersion;
+
+	// Is the device authorized for USB communication?  If not, then none of the other properties besides the serial number will be valid
+	bool bUnauthorizedDevice;
+
+	FAndroidDeviceInfo()
+		: SDKVersion(INDEX_NONE)
+		, GLESVersion(INDEX_NONE)
+		, bUnauthorizedDevice(false)
+	{
+	}
 };
 
 
diff --git a/Engine/Source/Developer/Android/AndroidTargetPlatform/Private/AndroidTargetDevice.h b/Engine/Source/Developer/Android/AndroidTargetPlatform/Private/AndroidTargetDevice.h
--- a/Engine/Source/Developer/Android/AndroidTargetPlatform/Private/AndroidTargetDevice.h
+++ b/Engine/Source/Developer/Android/AndroidTargetPlatform/Private/AndroidTargetDevice.h
@@ -21,8 +21,7 @@ typedef TSharedRef<class FAndroidTargetDevice, ESPMode::ThreadSafe> FAndroidTarg
 /**
  * Implements a Android target device.
  */
-class FAndroidTargetDevice
-	: public ITargetDevice
+class FAndroidTargetDevice : public ITargetDevice
 {
 public:
 
@@ -33,9 +32,11 @@ public:
 	 * @param InSerialNumber - The ADB serial number of the target device.
 	 * @param InAndroidVariant - The variant of the Android platform, i.e. ATC, DXT or PVRTC.
 	 */
-	FAndroidTargetDevice( const ITargetPlatform& InTargetPlatform, const FString& InSerialNumber, const FString& InAndroidVariant )
+	FAndroidTargetDevice(const ITargetPlatform& InTargetPlatform, const FString& InSerialNumber, const FString& InAndroidVariant)
 		: AndroidVariant(InAndroidVariant)
-		, Connected(false)
+		, bConnected(false)
+		, bIsDeviceAuthorized(false)
+		, AndroidSDKVersion(INDEX_NONE)
 		, DeviceName(InSerialNumber)
 		, Model(InSerialNumber)
 		, SerialNumber(InSerialNumber)
@@ -47,11 +48,33 @@ public:
 	/**
 	 * Sets the device's connection state.
 	 *
-	 * @param InConnected - Whether the device is connected.
+	 * @param bInConnected - Whether the device is connected.
 	 */
-	void SetConnected( bool InConnected )
+	void SetConnected(bool bInConnected)
 	{
-		Connected = InConnected;
+		bConnected = bInConnected;
+	}
+
+	/**
+	 * Sets the device's authorization state.
+	 *
+	 * @param bInConnected - Whether the device is authorized for USB communications.
+	 */
+	void SetAuthorized(bool bInIsAuthorized)
+	{
+		bIsDeviceAuthorized = bInIsAuthorized;
+	}
+
+	/**
+	 * Sets the device's OS/SDK versions.
+	 *
+	 * @param InSDKVersion - Android SDK version of the device.
+	 * @param InReleaseVersion - Android Release (human-readable) version of the device.
+	 */
+	void SetVersions(int32 InSDKVersion, const FString& InReleaseVersion)
+	{
+		AndroidSDKVersion = InSDKVersion;
+		AndroidVersionString = InReleaseVersion;
 	}
 
 	/**
@@ -59,7 +82,7 @@ public:
 	 *
 	 * @param InDeviceName - The device name to set.
 	 */
-	void SetDeviceName( const FString& InDeviceName )
+	void SetDeviceName(const FString& InDeviceName)
 	{
 		DeviceName = InDeviceName;
 	}
@@ -69,7 +92,7 @@ public:
 	 *
 	 * @param InDeviceName - The device name to set.
 	 */
-	void SetModel( const FString& InNodel )
+	void SetModel(const FString& InNodel)
 	{
 		Model = InNodel;
 	}
@@ -77,28 +100,29 @@ public:
 public:
 
 	// Begin ITargetDevice interface
-
-	virtual bool Connect( ) override
+	virtual bool Connect() override
 	{
 		return true;
 	}
 
-	virtual bool Deploy( const FString& SourceFolder, FString& OutAppId ) override;
+	virtual bool Deploy(const FString& SourceFolder, FString& OutAppId) override;
 
-	virtual void Disconnect( ) override
-	{ }
+	virtual void Disconnect() override
+	{
+	}
 
-	virtual ETargetDeviceTypes::Type GetDeviceType( ) const override
+	virtual ETargetDeviceTypes::Type GetDeviceType() const override
 	{
+		//@TODO: How to distinguish between a Tablet and a Phone (or a TV microconsole, etc...), and is it important?
 		return ETargetDeviceTypes::Tablet;
 	}
 
-	virtual FTargetDeviceId GetId( ) const override
+	virtual FTargetDeviceId GetId() const override
 	{
 		return FTargetDeviceId(TargetPlatform.PlatformName(), SerialNumber);
 	}
 
-	virtual FString GetName( ) const override
+	virtual FString GetName() const override
 	{
 		if (AndroidVariant.IsEmpty())
 		{
@@ -110,10 +134,7 @@ public:
 		return FString::Printf(TEXT("%s (%s)"), *Model, *AndroidVariant);
 	}
 
-	virtual FString GetOperatingSystemName( ) override
-	{
-		return TEXT("Android");
-	}
+	virtual FString GetOperatingSystemName() override;
 
 	virtual int32 GetProcessSnapshot( TArray<FTargetDeviceProcessInfo>& OutProcessInfos ) override;
 
@@ -122,39 +143,31 @@ public:
 		return TargetPlatform;
 	}
 
-	virtual bool IsConnected( ) override
+	virtual bool IsConnected() override
 	{
-		return Connected;
+		return bConnected;
 	}
 
-	virtual bool IsDefault( ) const override
+	virtual bool IsDefault() const override
 	{
 		return true;
 	}
 
-	virtual bool PowerOff( bool Force ) override;
+	virtual bool PowerOff(bool Force) override;
 
-	virtual bool PowerOn( ) override
+	virtual bool PowerOn() override
 	{
 		return true;
 	}
 
-	virtual bool Launch( const FString& AppId, EBuildConfigurations::Type BuildConfiguration, EBuildTargets::Type BuildTarget, const FString& Params, uint32* OutProcessId );
-
-	virtual bool Reboot( bool bReconnect = false ) override;
-
-	virtual bool Run( const FString& ExecutablePath, const FString& Params, uint32* OutProcessId ) override;
-
-	virtual bool SupportsFeature( ETargetDeviceFeatures::Type Feature ) const override;
-
-	virtual bool SupportsSdkVersion( const FString& VersionString ) const override;
-
-	virtual bool TerminateProcess( const int32 ProcessId ) override;
-
-	virtual void SetUserCredentials( const FString & UserName, const FString & UserPassword ) override;
-
-	virtual bool GetUserCredentials( FString & OutUserName, FString & OutUserPassword ) override;
-
+	virtual bool Launch(const FString& AppId, EBuildConfigurations::Type BuildConfiguration, EBuildTargets::Type BuildTarget, const FString& Params, uint32* OutProcessId);
+	virtual bool Reboot(bool bReconnect = false) override;
+	virtual bool Run(const FString& ExecutablePath, const FString& Params, uint32* OutProcessId) override;
+	virtual bool SupportsFeature(ETargetDeviceFeatures::Type Feature) const override;
+	virtual bool SupportsSdkVersion(const FString& VersionString) const override;
+	virtual bool TerminateProcess(const int32 ProcessId) override;
+	virtual void SetUserCredentials(const FString & UserName, const FString & UserPassword) override;
+	virtual bool GetUserCredentials(FString & OutUserName, FString & OutUserPassword) override;
 	// End ITargetDevice interface
 
 protected:
@@ -176,7 +189,16 @@ private:
 	FString AndroidVariant;
 
 	// Holds a flag indicating whether the device is currently connected.
-	bool Connected;
+	bool bConnected;
+
+	// Holds a flag indicating whether the device is USB comms authorized (if not, most other values aren't valid but we still want to show the device as detected but unready)
+	bool bIsDeviceAuthorized;
+
+	// Holds the Android SDK version
+	int32 AndroidSDKVersion;
+
+	// Holds the Android Release version string (e.g., "2.3" or "4.2.2")
+	FString AndroidVersionString;
 
 	// Holds the device name.
 	FString DeviceName;
diff --git a/Engine/Source/Developer/Android/AndroidTargetPlatform/Private/AndroidTargetDevice.inl b/Engine/Source/Developer/Android/AndroidTargetPlatform/Private/AndroidTargetDevice.inl
--- a/Engine/Source/Developer/Android/AndroidTargetPlatform/Private/AndroidTargetDevice.inl
+++ b/Engine/Source/Developer/Android/AndroidTargetPlatform/Private/AndroidTargetDevice.inl
@@ -1,9 +1,5 @@
 // Copyright 1998-2014 Epic Games, Inc. All Rights Reserved.
 
-/*=============================================================================
-	AndroidTargetDevice.cpp: Implements the AndroidTargetDevice class.
-=============================================================================*/
-
 
 /* ITargetDevice interface
  *****************************************************************************/
@@ -24,6 +20,17 @@ inline bool FAndroidTargetDevice::Deploy( const FString& SourceFolder, FString&
 	return 0 == ReturnCode;
 }
 
+inline FString FAndroidTargetDevice::GetOperatingSystemName()
+{
+	if (!AndroidVersionString.IsEmpty())
+	{
+		return FString::Printf(TEXT("Android %s, API level %d"), *AndroidVersionString, AndroidSDKVersion);
+	}
+	else
+	{
+		return TEXT("Android");
+	}
+}
 
 inline int32 FAndroidTargetDevice::GetProcessSnapshot( TArray<FTargetDeviceProcessInfo>& OutProcessInfos ) 
 {
diff --git a/Engine/Source/Developer/Android/AndroidTargetPlatform/Private/AndroidTargetPlatform.inl b/Engine/Source/Developer/Android/AndroidTargetPlatform/Private/AndroidTargetPlatform.inl
--- a/Engine/Source/Developer/Android/AndroidTargetPlatform/Private/AndroidTargetPlatform.inl
+++ b/Engine/Source/Developer/Android/AndroidTargetPlatform/Private/AndroidTargetPlatform.inl
@@ -384,6 +384,8 @@ inline bool FAndroidTargetPlatform<TPlatformProperties>::HandleTicker( float Del
 			Device->SetConnected(true);
 			Device->SetModel(DeviceInfo.Model);
 			Device->SetDeviceName(DeviceInfo.DeviceName);
+			Device->SetAuthorized(!DeviceInfo.bUnauthorizedDevice);
+			Device->SetVersions(DeviceInfo.SDKVersion, DeviceInfo.HumanAndroidVersion);
 
 			DeviceDiscoveredEvent.Broadcast(Device.ToSharedRef());
 		}
