Mac build fix
#mac

diff --git a/Engine/Source/Developer/Android/AndroidPlatformEditor/Private/AndroidTargetSettingsCustomization.cpp b/Engine/Source/Developer/Android/AndroidPlatformEditor/Private/AndroidTargetSettingsCustomization.cpp
--- a/Engine/Source/Developer/Android/AndroidPlatformEditor/Private/AndroidTargetSettingsCustomization.cpp
+++ b/Engine/Source/Developer/Android/AndroidPlatformEditor/Private/AndroidTargetSettingsCustomization.cpp
@@ -27,7 +27,8 @@ TSharedRef<IDetailCustomization> FAndroidTargetSettingsCustomization::MakeInstan
 }
 
 FAndroidTargetSettingsCustomization::FAndroidTargetSettingsCustomization()
-	: EngineAndroidPath(FPaths::EngineDir() + TEXT("Build/Android/Java"))
+	: AndroidRelativePath(TEXT(""))
+	, EngineAndroidPath(FPaths::EngineDir() + TEXT("Build/Android/Java"))
 	, GameAndroidPath(FPaths::GameDir() + TEXT("Build/Android"))
 	, EngineManifestPath(EngineAndroidPath / TEXT("AndroidManifest.xml"))
 	, GameManifestPath(GameAndroidPath / TEXT("AndroidManifest.xml"))
