On Mac, read Android SDK paths from .bash_profile and set evn variables using them

diff --git a/Engine/Source/Developer/Android/AndroidDeviceDetection/Private/AndroidDeviceDetectionModule.cpp b/Engine/Source/Developer/Android/AndroidDeviceDetection/Private/AndroidDeviceDetectionModule.cpp
--- a/Engine/Source/Developer/Android/AndroidDeviceDetection/Private/AndroidDeviceDetectionModule.cpp
+++ b/Engine/Source/Developer/Android/AndroidDeviceDetection/Private/AndroidDeviceDetectionModule.cpp
@@ -255,6 +255,56 @@ public:
 		TCHAR AndroidDirectory[32768] = { 0 };
 		FPlatformMisc::GetEnvironmentVariable(TEXT("ANDROID_HOME"), AndroidDirectory, 32768);
 
+#if PLATFORM_MAC
+		TCHAR AntDirectory[32768] = { 0 };
+		TCHAR NDKDirectory[32768] = { 0 };
+		FPlatformMisc::GetEnvironmentVariable(TEXT("ANT_HOME"), AntDirectory, 32768);
+		FPlatformMisc::GetEnvironmentVariable(TEXT("NDKROOT"), NDKDirectory, 32768);
+
+		if (AndroidDirectory[0] == 0 || AntDirectory[0] == 0 || NDKDirectory[0] == 0)
+		{
+			FArchive* FileReader = IFileManager::Get().CreateFileReader(*FString([@"~/.bash_profile" stringByExpandingTildeInPath]));
+			if (FileReader)
+			{
+				const int64 FileSize = FileReader->TotalSize();
+				ANSICHAR* AnsiContents = (ANSICHAR*)FMemory::Malloc(FileSize);
+				FileReader->Serialize(AnsiContents, FileSize);
+				FileReader->Close();
+				delete FileReader;
+
+				TArray<FString> Lines;
+				FString(ANSI_TO_TCHAR(AnsiContents)).ParseIntoArrayLines(&Lines);
+				FMemory::Free(AnsiContents);
+
+				for (int32 Index = 0; Index < Lines.Num(); Index++)
+				{
+					if (AndroidDirectory[0] == 0 && Lines[Index].StartsWith(TEXT("export ANDROID_HOME=")))
+					{
+						FString Directory;
+						Lines[Index].Split(TEXT("="), NULL, &Directory);
+						Directory = Directory.Replace(TEXT("\""), TEXT(""));
+						FCString::Strcpy(AndroidDirectory, *Directory);
+						setenv("ANDROID_HOME", TCHAR_TO_ANSI(AndroidDirectory), 1);
+					}
+					else if (AntDirectory[0] == 0 && Lines[Index].StartsWith(TEXT("export ANT_HOME=")))
+					{
+						FString Directory;
+						Lines[Index].Split(TEXT("="), NULL, &Directory);
+						Directory = Directory.Replace(TEXT("\""), TEXT(""));
+						setenv("ANT_HOME", TCHAR_TO_ANSI(*Directory), 1);
+					}
+					else if (NDKDirectory[0] == 0 && Lines[Index].StartsWith(TEXT("export NDKROOT=")))
+					{
+						FString Directory;
+						Lines[Index].Split(TEXT("="), NULL, &Directory);
+						Directory = Directory.Replace(TEXT("\""), TEXT(""));
+						setenv("NDKROOT", TCHAR_TO_ANSI(*Directory), 1);
+					}
+				}
+			}
+		}
+#endif
+
 		if (AndroidDirectory[0] == 0)
 		{
 			return;
diff --git a/Engine/Source/Programs/AutomationTool/Android/AndroidPlatform.Automation.cs b/Engine/Source/Programs/AutomationTool/Android/AndroidPlatform.Automation.cs
--- a/Engine/Source/Programs/AutomationTool/Android/AndroidPlatform.Automation.cs
+++ b/Engine/Source/Programs/AutomationTool/Android/AndroidPlatform.Automation.cs
@@ -18,6 +18,21 @@ public class AndroidPlatform : Platform
 	public AndroidPlatform()
 		: base(UnrealTargetPlatform.Android)
 	{
+		bool bNeedsAndroidHome = string.IsNullOrEmpty(Environment.GetEnvironmentVariable("ANDROID_HOME"));
+		if (Utils.IsRunningOnMono && bNeedsAndroidHome)
+		{
+			// Try reading env variable we need from .bash_profile
+			string BashProfilePath = Path.Combine(Environment.GetFolderPath(Environment.SpecialFolder.Personal), ".bash_profile");
+			string[] BashProfileContents = File.ReadAllLines(BashProfilePath);
+			foreach (string Line in BashProfileContents)
+			{
+				if (Line.StartsWith("export ANDROID_HOME="))
+				{
+					string PathVar = Line.Split('=')[1].Replace("\"", "");
+					Environment.SetEnvironmentVariable("ANDROID_HOME", PathVar);
+				}
+			}
+		}
 	}
 
 	private static string GetSONameWithoutArchitecture(ProjectParams Params, string DecoratedExeName)
diff --git a/Engine/Source/Programs/UnrealBuildTool/Android/UEBuildAndroid.cs b/Engine/Source/Programs/UnrealBuildTool/Android/UEBuildAndroid.cs
--- a/Engine/Source/Programs/UnrealBuildTool/Android/UEBuildAndroid.cs
+++ b/Engine/Source/Programs/UnrealBuildTool/Android/UEBuildAndroid.cs
@@ -64,6 +64,35 @@ namespace UnrealBuildTool
         private bool HasAnySDK()
         {
             string NDKPath = Environment.GetEnvironmentVariable("NDKROOT");
+			bool bNeedsNDKPath = string.IsNullOrEmpty(NDKPath);
+			bool bNeedsAndroidHome = string.IsNullOrEmpty(Environment.GetEnvironmentVariable("ANDROID_HOME"));
+			bool bNeedsAntHome = string.IsNullOrEmpty(Environment.GetEnvironmentVariable("ANT_HOME"));
+
+			if (Utils.IsRunningOnMono && (bNeedsNDKPath || bNeedsAndroidHome || bNeedsAntHome))
+			{
+				// Try reading env variables we need from .bash_profile
+				string BashProfilePath = Path.Combine(Environment.GetFolderPath(Environment.SpecialFolder.Personal), ".bash_profile");
+				string[] BashProfileContents = File.ReadAllLines(BashProfilePath);
+				foreach (string Line in BashProfileContents)
+				{
+					if (bNeedsAndroidHome && Line.StartsWith("export ANDROID_HOME="))
+					{
+						string PathVar = Line.Split('=')[1].Replace("\"", "");
+						Environment.SetEnvironmentVariable("ANDROID_HOME", PathVar);
+					}
+					else if (bNeedsNDKPath && Line.StartsWith("export NDKROOT="))
+					{
+						string PathVar = Line.Split('=')[1].Replace("\"", "");
+						Environment.SetEnvironmentVariable("NDKROOT", PathVar);
+						NDKPath = PathVar;
+					}
+					else if (bNeedsAntHome && Line.StartsWith("export ANT_HOME="))
+					{
+						string PathVar = Line.Split('=')[1].Replace("\"", "");
+						Environment.SetEnvironmentVariable("ANT_HOME", PathVar);
+					}
+				}
+			}
 
             // we don't have an NDKROOT specified
             if (String.IsNullOrEmpty(NDKPath))
