Fixed Android SDKPath fallback to ANDROID_HOME environment variable
#ue4
#android

diff --git a/Engine/Source/Developer/Android/AndroidPlatformEditor/Private/AndroidSDKSettings.cpp b/Engine/Source/Developer/Android/AndroidPlatformEditor/Private/AndroidSDKSettings.cpp
--- a/Engine/Source/Developer/Android/AndroidPlatformEditor/Private/AndroidSDKSettings.cpp
+++ b/Engine/Source/Developer/Android/AndroidPlatformEditor/Private/AndroidSDKSettings.cpp
@@ -43,16 +43,8 @@ void UAndroidSDKSettings::SetupInitialTargetPaths()
 		{
 			SDKPath.Path = TempPath.Mid(7, TempPath.Len() - 9); // Pull the value from the string Path="<payload>"
 		}
-		else
-		{
-			SDKPath.Path = "Bar";
-		}
 	}
-	else
-	{
-		SDKPath.Path = "Foo";
-	}
-	UE_LOG(AndroidSDKSettings, Warning, TEXT("SDK Path = %s"), *SDKPath.Path);
+//	UE_LOG(AndroidSDKSettings, Warning, TEXT("SDK Path = %s"), *SDKPath.Path);
 
 	if (GConfig->GetString(TEXT("/Script/AndroidPlatformEditor.AndroidSDKSettings"), TEXT("NDKPath"), TempPath, UserIni))
 	{
