- No longer using curly brackets with {PROJECT_NAME} default, messes up in editor
- Fixed some tooltips
- Missed a file in changelist 2414384

diff --git a/Engine/Config/BaseEngine.ini b/Engine/Config/BaseEngine.ini
--- a/Engine/Config/BaseEngine.ini
+++ b/Engine/Config/BaseEngine.ini
@@ -2142,7 +2142,7 @@ bBuildForGL4=false
 bSplitIntoSeparateApks=false
 Orientation=landscape
 DepthBufferPreference=0
-PackageName=com.YourCompany.{PROJECT_NAME}
+PackageName=com.YourCompany.[PROJECT]
 StoreVersion=1
 VersionDisplayName=1.0
 MinSDKVersion=9
diff --git a/Engine/Source/Developer/Android/AndroidPlatformEditor/Private/AndroidSDKSettings.h b/Engine/Source/Developer/Android/AndroidPlatformEditor/Private/AndroidSDKSettings.h
--- a/Engine/Source/Developer/Android/AndroidPlatformEditor/Private/AndroidSDKSettings.h
+++ b/Engine/Source/Developer/Android/AndroidPlatformEditor/Private/AndroidSDKSettings.h
@@ -16,15 +16,27 @@ class ANDROIDPLATFORMEDITOR_API UAndroidSDKSettings : public UObject
 public:
 	GENERATED_UCLASS_BODY()
 
+	// Location on disk to find the Android SDK (defaults to ANDROID_HOME environment variable if set)
 	UPROPERTY(GlobalConfig, EditAnywhere, Category = SDKConfig, Meta = (DisplayName = "Location of Android SDK (the directory usually contains 'android-sdk-')"))
 	FDirectoryPath SDKPath;
 
+	// Location on disk to find the Android NDK (defaults to NDKROOT environment variable if set)
 	UPROPERTY(GlobalConfig, EditAnywhere, Category = SDKConfig, Meta = (DisplayName = "Location of Android NDK (the directory usually contains 'android-ndk-')"))
 	FDirectoryPath NDKPath;
 
+	// Location on disk to find the ANT tool
 	UPROPERTY(GlobalConfig, EditAnywhere, Category = SDKConfig, Meta = (DisplayName = "Location of ANT (the directory usually contains 'apache-ant-')"))
 	FDirectoryPath ANTPath;
 
+	// Which SDK to package and compile Java with (a specific version or (without quotes) 'latest' for latest version on disk, or 'matchndk' to match the NDK API Level)
+	UPROPERTY(GlobalConfig, EditAnywhere, Category = SDKConfig, Meta = (DisplayName = "SDK API Level (specific version, 'latest', or 'matchndk' - see tooltip)"))
+	FString SDKAPILevel;
+
+	// Which NDK to compile with (a specific version or (without quotes) 'latest' for latest version on disk). Note that choosing android-21 or later won't run on pre-5.0 devices.
+	UPROPERTY(GlobalConfig, EditAnywhere, Category = SDKConfig, Meta = (DisplayName = "NDK API Level (specific version or 'latest' - see tooltip)"))
+	FString NDKAPILevel;
+
+
 #if WITH_EDITOR
 	// UObject interface
 	virtual void PostEditChangeProperty(struct FPropertyChangedEvent& PropertyChangedEvent) override;
diff --git a/Engine/Source/Programs/UnrealBuildTool/Android/UEDeployAndroid.cs b/Engine/Source/Programs/UnrealBuildTool/Android/UEDeployAndroid.cs
--- a/Engine/Source/Programs/UnrealBuildTool/Android/UEDeployAndroid.cs
+++ b/Engine/Source/Programs/UnrealBuildTool/Android/UEDeployAndroid.cs
@@ -656,7 +656,7 @@ namespace UnrealBuildTool.Android
 			Ini.GetArray("/Script/AndroidRuntimeSettings.AndroidRuntimeSettings", "ExtraPermissions", out ExtraPermissions);
 
 			// replace some variables
-			PackageName = PackageName.Replace("{PROJECT_NAME}", ProjectName);
+			PackageName = PackageName.Replace("[PROJECT]", ProjectName);
 
 			StringBuilder Text = new StringBuilder();
 			Text.AppendLine("<?xml version=\"1.0\" encoding=\"utf-8\"?>");
diff --git a/Engine/Source/Runtime/Android/AndroidRuntimeSettings/Classes/AndroidRuntimeSettings.h b/Engine/Source/Runtime/Android/AndroidRuntimeSettings/Classes/AndroidRuntimeSettings.h
--- a/Engine/Source/Runtime/Android/AndroidRuntimeSettings/Classes/AndroidRuntimeSettings.h
+++ b/Engine/Source/Runtime/Android/AndroidRuntimeSettings/Classes/AndroidRuntimeSettings.h
@@ -105,7 +105,7 @@ public:
 	GENERATED_UCLASS_BODY()
 
 	// The official name of the product (same as the name you use on the Play Store web site)
-	UPROPERTY(GlobalConfig, EditAnywhere, Category = APKPackaging, Meta = (DisplayName = "Android Package Name (usually in the 'com.Company.Project' format)"))
+	UPROPERTY(GlobalConfig, EditAnywhere, Category = APKPackaging, Meta = (DisplayName = "Android Package Name ('com.Company.Project', [PROJECT] is replaced with project name)"))
 	FString PackageName;
 
 	// The version number used to indicate newer versions in the Store
@@ -140,8 +140,8 @@ public:
 	UPROPERTY(GlobalConfig, EditAnywhere, Category = AdvancedAPKPackaging, Meta = (DisplayName = "Extra Permissions (e.g. 'android.permission.INTERNET')"))
 	TArray<FString> ExtraPermissions;
 	
-	// This is the file that keytool outputs, specified with the -keystore parameter (absolute path or relative to <Project>/Build/Android)
-	UPROPERTY(GlobalConfig, EditAnywhere, Category = DistributionSigning, Meta = (DisplayName = "Key Store (output of keytool)"))
+	// This is the file that keytool outputs, specified with the -keystore parameter (file should be in <Project>/Build/Android)
+	UPROPERTY(GlobalConfig, EditAnywhere, Category = DistributionSigning, Meta = (DisplayName = "Key Store (output of keytool, placed in <Project>/Build/Android)"))
 	FString KeyStore;
 	
 	// This is the name of the key that you specified with the -alias parameter to keytool
