Adding Java_Home support to ini files for Android building.

diff --git a/Engine/Source/Developer/Android/AndroidPlatformEditor/Private/AndroidSDKSettings.cpp b/Engine/Source/Developer/Android/AndroidPlatformEditor/Private/AndroidSDKSettings.cpp
--- a/Engine/Source/Developer/Android/AndroidPlatformEditor/Private/AndroidSDKSettings.cpp
+++ b/Engine/Source/Developer/Android/AndroidPlatformEditor/Private/AndroidSDKSettings.cpp
@@ -32,34 +32,6 @@ void UAndroidSDKSettings::SetDeviceDetection(IAndroidDeviceDetection * AndroidDe
 
 void UAndroidSDKSettings::SetupInitialTargetPaths()
 {
-	// Temp code to pull directly from GConfig until the correct serialisation code exists
-	const FString UserIni = GEngineIni; //  (TEXT("DefaultUser"));
-	FString TempPath;
-	if (GConfig->GetString(TEXT("/Script/AndroidPlatformEditor.AndroidSDKSettings"), TEXT("SDKPath"), TempPath, UserIni))
-	{
-		if (TempPath.Contains(TEXT("Path=")))
-		{
-			SDKPath.Path = TempPath.Mid(7, TempPath.Len() - 9); // Pull the value from the string Path="<payload>"
-		}
-	}
-
-	if (GConfig->GetString(TEXT("/Script/AndroidPlatformEditor.AndroidSDKSettings"), TEXT("NDKPath"), TempPath, UserIni))
-	{
-		if (TempPath.Contains(TEXT("Path=")))
-		{
-			NDKPath.Path = TempPath.Mid(7, TempPath.Len() - 9); // Pull the value from the string Path="<payload>"
-		}
-	}
-
-	if (GConfig->GetString(TEXT("/Script/AndroidPlatformEditor.AndroidSDKSettings"), TEXT("ANTPath"), TempPath, UserIni))
-	{
-		if (TempPath.Contains(TEXT("Path=")))
-		{
-			ANTPath.Path = TempPath.Mid(7, TempPath.Len() - 9); // Pull the value from the string Path="<payload>"
-		}
-	}
-	// End Temp section - this will be removed asap
-
 	// Check each path, if we don't have it then try to look up the values from a known location instead
 	if (SDKPath.Path.IsEmpty())
 	{
@@ -91,6 +63,16 @@ void UAndroidSDKSettings::SetupInitialTargetPaths()
 		}
 	}
 
+	if (JavaPath.Path.IsEmpty())
+	{
+		TCHAR AndroidJavaPath[256];
+		FPlatformMisc::GetEnvironmentVariable(TEXT("JAVA_HOME"), AndroidJavaPath, ARRAY_COUNT(AndroidJavaPath));
+		if (AndroidJavaPath[0])
+		{
+			JavaPath.Path = AndroidJavaPath;
+		}
+	}
+
 
 #if PLATFORM_MAC
 	// On a Mac, if we still don't have the values then check the .bash_profile file for them
@@ -159,6 +141,12 @@ void UAndroidSDKSettings::UpdateTargetModulePaths(bool bForceUpdate)
 		Values.Add(ANTPath.Path);
 	}
 
+	if (bForceUpdate || !JavaPath.Path.IsEmpty())
+	{
+		Keys.Add(TEXT("JAVA_HOME"));
+		Values.Add(JavaPath.Path);
+	}
+
 	SaveConfig();
 	
 	if (Keys.Num() != 0)
diff --git a/Engine/Source/Developer/Android/AndroidPlatformEditor/Private/AndroidSDKSettings.h b/Engine/Source/Developer/Android/AndroidPlatformEditor/Private/AndroidSDKSettings.h
--- a/Engine/Source/Developer/Android/AndroidPlatformEditor/Private/AndroidSDKSettings.h
+++ b/Engine/Source/Developer/Android/AndroidPlatformEditor/Private/AndroidSDKSettings.h
@@ -36,6 +36,8 @@ public:
 	UPROPERTY(GlobalConfig, EditAnywhere, Category = SDKConfig, Meta = (DisplayName = "NDK API Level (specific version or 'latest' - see tooltip)"))
 	FString NDKAPILevel;
 
+	UPROPERTY(GlobalConfig, EditAnywhere, Category = SDKConfig, Meta = (DisplayName = "Location of JAVA (the directory usually contains 'jdk')"))
+	FDirectoryPath JavaPath;
 
 #if WITH_EDITOR
 	// UObject interface
diff --git a/Engine/Source/Developer/Android/AndroidPlatformEditor/Private/AndroidSDKSettingsCustomization.cpp b/Engine/Source/Developer/Android/AndroidPlatformEditor/Private/AndroidSDKSettingsCustomization.cpp
--- a/Engine/Source/Developer/Android/AndroidPlatformEditor/Private/AndroidSDKSettingsCustomization.cpp
+++ b/Engine/Source/Developer/Android/AndroidPlatformEditor/Private/AndroidSDKSettingsCustomization.cpp
@@ -62,8 +62,7 @@ void FAndroidSDKSettingsCustomization::SetupSDKPaths()
 			changed |= true;
 		}
 				
-	}
-	
+	}	
 
 	if (settings->NDKPath.Path.IsEmpty())
 	{
@@ -75,8 +74,7 @@ void FAndroidSDKSettingsCustomization::SetupSDKPaths()
 			changed |= true;
 		}
 		
-	}
-	
+	}	
 
 	if (settings->ANTPath.Path.IsEmpty())
 	{
@@ -89,6 +87,18 @@ void FAndroidSDKSettingsCustomization::SetupSDKPaths()
 		}
 		
 	}
+
+	if (settings->JavaPath.Path.IsEmpty())
+	{
+		TCHAR AndroidJavaPath[256];
+		FPlatformMisc::GetEnvironmentVariable(TEXT("JAVA_HOME"), AndroidJavaPath, ARRAY_COUNT(AndroidJavaPath));
+		if (AndroidJavaPath[0])
+		{
+			settings->JavaPath.Path = AndroidJavaPath;
+			changed |= true;
+		}
+
+	}
 	
 	if (changed)
 	{
diff --git a/Engine/Source/Developer/Android/AndroidTargetPlatform/Private/AndroidTargetPlatform.inl b/Engine/Source/Developer/Android/AndroidTargetPlatform/Private/AndroidTargetPlatform.inl
--- a/Engine/Source/Developer/Android/AndroidTargetPlatform/Private/AndroidTargetPlatform.inl
+++ b/Engine/Source/Developer/Android/AndroidTargetPlatform/Private/AndroidTargetPlatform.inl
@@ -91,19 +91,6 @@ inline bool FAndroidTargetPlatform<TPlatformProperties>::IsSdkInstalled(bool bPr
 {
 
 	OutDocumentationPath = FString("Shared/Tutorials/SettingUpAndroidTutorial");
-
-#if PLATFORM_WINDOWS
-	TCHAR JAVA_HOME[MAX_PATH];
-	FPlatformMisc::GetEnvironmentVariable(TEXT("JAVA_HOME"), JAVA_HOME, MAX_PATH);
-
-	// make sure that JAVA_HOME points to the right thing
-	if (JAVA_HOME[0] == 0 ||
-		IFileManager::Get().FileSize(*(FString(JAVA_HOME) / TEXT("bin/javac.exe"))) < 0)
-	{
-		return false;
-	}
-#endif
-
 	return true;
 }
 
diff --git a/Engine/Source/Programs/UnrealBuildTool/Android/UEBuildAndroid.cs b/Engine/Source/Programs/UnrealBuildTool/Android/UEBuildAndroid.cs
--- a/Engine/Source/Programs/UnrealBuildTool/Android/UEBuildAndroid.cs
+++ b/Engine/Source/Programs/UnrealBuildTool/Android/UEBuildAndroid.cs
@@ -68,7 +68,8 @@ namespace UnrealBuildTool
                 Dictionary<string, string> EnvVarNames = new Dictionary<string,string> { 
                                                          {"ANDROID_HOME", "SDKPath"}, 
                                                          {"NDKROOT", "NDKPath"}, 
-                                                         {"ANT_HOME", "ANTPath"}
+                                                         {"ANT_HOME", "ANTPath"},
+                                                         {"JAVA_HOME", "JavaPath"}
                                                          };
 
                 string path;
