[INTEGRATE] Change 2418497 by [REDACTED_PERSON][REDACTED_HANDLE] on 2015/01/26 06:46:00

	UE-8075 - Devices not updating post-SDK path change
	Fixed code so that it correctly references the device detection class.
	Added a 'force update' path so that upon successful SDK path updates the devices are rechecked right away instead of waiting up to 10 seconds.
	Cleared out old devices if the SDK path is empty/doesn't contain the adb executable

	Change to UEBuildAndroid so that empty SDK paths are no longer treated as valid paths on start up allowing things like SDKPath=("") to be overridden by env vars.

diff --git a/Engine/Source/Developer/Android/AndroidDeviceDetection/Private/AndroidDeviceDetectionModule.cpp b/Engine/Source/Developer/Android/AndroidDeviceDetection/Private/AndroidDeviceDetectionModule.cpp
--- a/Engine/Source/Developer/Android/AndroidDeviceDetection/Private/AndroidDeviceDetectionModule.cpp
+++ b/Engine/Source/Developer/Android/AndroidDeviceDetection/Private/AndroidDeviceDetectionModule.cpp
@@ -8,6 +8,7 @@
 
 #define LOCTEXT_NAMESPACE "FAndroidDeviceDetectionModule" 
 
+DEFINE_LOG_CATEGORY_STATIC(AndroidDeviceDetectionLog, Log, All);
 
 class FAndroidDeviceDetectionRunnable : public FRunnable
 {
@@ -17,7 +18,8 @@ public:
 		DeviceMap(InDeviceMap),
 		DeviceMapLock(InDeviceMapLock),
 		ADBPathCheckLock(InADBPathCheckLock),
-		HasADBPath(false)
+		HasADBPath(false),
+		ForceCheck(false)
 
 	{
 	}
@@ -46,7 +48,7 @@ public:
 		while (StopTaskCounter.GetValue() == 0)
 		{
 			// query every 10 seconds
-			if (LoopCount++ >= 10)
+			if (LoopCount++ >= 10 || ForceCheck)
 			{
 				// Make sure we have an ADB path before checking
 				FScopeLock PathLock(ADBPathCheckLock);
@@ -54,6 +56,7 @@ public:
 					QueryConnectedDevices();
 
 				LoopCount = 0;
+				ForceCheck = false;
 			}
 
 			FPlatformProcess::Sleep(1.0f);
@@ -66,6 +69,14 @@ public:
 	{
 		ADBPath = InADBPath;
 		HasADBPath = !ADBPath.IsEmpty();
+		// Force a check next time we go around otherwise it can take over 10sec to find devices
+		ForceCheck = HasADBPath;	
+
+		// If we have no path then clean the existing devices out
+		if (!HasADBPath && DeviceMap.Num() > 0)
+		{
+			DeviceMap.Reset();
+		}
 	}
 
 private:
@@ -256,6 +267,7 @@ private:
 
 	FCriticalSection* ADBPathCheckLock;
 	bool HasADBPath;
+	bool ForceCheck;
 };
 
 class FAndroidDeviceDetection : public IAndroidDeviceDetection
@@ -293,14 +305,14 @@ public:
 		return &DeviceMapLock;
 	}
 
-	virtual void UpdateADBPath()
+	virtual void UpdateADBPath() override
 	{
 		FScopeLock PathUpdateLock(&ADBPathCheckLock);
 		TCHAR AndroidDirectory[32768] = { 0 };
 		FPlatformMisc::GetEnvironmentVariable(TEXT("ANDROID_HOME"), AndroidDirectory, 32768);
 
 		FString ADBPath;
-
+		
 		if (AndroidDirectory[0] != 0)
 		{
 #if PLATFORM_WINDOWS
diff --git a/Engine/Source/Developer/Android/AndroidPlatformEditor/Private/AndroidPlatformEditorModule.cpp b/Engine/Source/Developer/Android/AndroidPlatformEditor/Private/AndroidPlatformEditorModule.cpp
--- a/Engine/Source/Developer/Android/AndroidPlatformEditor/Private/AndroidPlatformEditorModule.cpp
+++ b/Engine/Source/Developer/Android/AndroidPlatformEditor/Private/AndroidPlatformEditorModule.cpp
@@ -57,8 +57,8 @@ class FAndroidPlatformEditorModule
 		auto &TargetPlatformManagerModule = FModuleManager::LoadModuleChecked<ITargetPlatformManagerModule>("TargetPlatform");
 		UAndroidSDKSettings * settings = GetMutableDefault<UAndroidSDKSettings>();
 		settings->SetTargetModule(&TargetPlatformManagerModule);
-		auto &AndroidDeviceDetection = FModuleManager::LoadModuleChecked<IAndroidDeviceDetection>("AndroidDeviceDetection");
-		settings->SetDeviceDetection(&AndroidDeviceDetection);
+		auto &AndroidDeviceDetection = FModuleManager::LoadModuleChecked<IAndroidDeviceDetectionModule>("AndroidDeviceDetection");
+		settings->SetDeviceDetection(AndroidDeviceDetection.GetAndroidDeviceDetection());
 		settings->SetupInitialTargetPaths();
 
 	}
diff --git a/Engine/Source/Programs/UnrealBuildTool/Android/UEBuildAndroid.cs b/Engine/Source/Programs/UnrealBuildTool/Android/UEBuildAndroid.cs
--- a/Engine/Source/Programs/UnrealBuildTool/Android/UEBuildAndroid.cs
+++ b/Engine/Source/Programs/UnrealBuildTool/Android/UEBuildAndroid.cs
@@ -75,8 +75,9 @@ namespace UnrealBuildTool
                 string path;
                 foreach(var kvp in EnvVarNames)
                 {
-                    if (configCacheIni.GetPath("/Script/AndroidPlatformEditor.AndroidSDKSettings", kvp.Value, out path))
+                    if (configCacheIni.GetPath("/Script/AndroidPlatformEditor.AndroidSDKSettings", kvp.Value, out path) && !string.IsNullOrEmpty(path))
                     {
+//                        Log.TraceWarning("Adding {0} from ini as {1} to {2}", kvp.Value, path, kvp.Key);
                         AndroidEnv.Add(kvp.Key, path);
                     }
                     else
@@ -84,6 +85,7 @@ namespace UnrealBuildTool
                         var envValue = Environment.GetEnvironmentVariable(kvp.Key);
                         if(!String.IsNullOrEmpty(envValue))
                         {
+//                            Log.TraceWarning("Adding {0} from env as {1}", kvp.Key, envValue);
                             AndroidEnv.Add(kvp.Key, envValue);
                         }
                     }
