[INTEGRATE] Change 2418550 by [REDACTED_PERSON][REDACTED_HANDLE] on 2015/01/26 08:53:09

	Removing JAVA_HOME from Mac Editor builds

diff --git a/Engine/Source/Developer/Android/AndroidPlatformEditor/Private/AndroidSDKSettings.cpp b/Engine/Source/Developer/Android/AndroidPlatformEditor/Private/AndroidSDKSettings.cpp
--- a/Engine/Source/Developer/Android/AndroidPlatformEditor/Private/AndroidSDKSettings.cpp
+++ b/Engine/Source/Developer/Android/AndroidPlatformEditor/Private/AndroidSDKSettings.cpp
@@ -63,6 +63,7 @@ void UAndroidSDKSettings::SetupInitialTargetPaths()
 		}
 	}
 
+#if PLATFORM_MAC == 0
 	if (JavaPath.Path.IsEmpty())
 	{
 		TCHAR AndroidJavaPath[256];
@@ -72,7 +73,7 @@ void UAndroidSDKSettings::SetupInitialTargetPaths()
 			JavaPath.Path = AndroidJavaPath;
 		}
 	}
-
+#endif
 
 #if PLATFORM_MAC
 	// On a Mac, if we still don't have the values then check the .bash_profile file for them
@@ -144,12 +145,14 @@ void UAndroidSDKSettings::UpdateTargetModulePaths(bool bForceUpdate)
 		Values.Add(ANTPath.Path);
 	}
 
+#if PLATFORM_MAC == 0
 	if (bForceUpdate || !JavaPath.Path.IsEmpty())
 	{
 		FPaths::NormalizeFilename(JavaPath.Path);
 		Keys.Add(TEXT("JAVA_HOME"));
 		Values.Add(JavaPath.Path);
 	}
+#endif
 
 	SaveConfig();
 	
diff --git a/Engine/Source/Developer/Android/AndroidPlatformEditor/Private/AndroidSDKSettings.h b/Engine/Source/Developer/Android/AndroidPlatformEditor/Private/AndroidSDKSettings.h
--- a/Engine/Source/Developer/Android/AndroidPlatformEditor/Private/AndroidSDKSettings.h
+++ b/Engine/Source/Developer/Android/AndroidPlatformEditor/Private/AndroidSDKSettings.h
@@ -36,8 +36,10 @@ public:
 	UPROPERTY(GlobalConfig, EditAnywhere, Category = SDKConfig, Meta = (DisplayName = "NDK API Level (specific version or 'latest' - see tooltip)"))
 	FString NDKAPILevel;
 
+#if PLATFORM_MAC == 0
 	UPROPERTY(GlobalConfig, EditAnywhere, Category = SDKConfig, Meta = (DisplayName = "Location of JAVA (the directory usually contains 'jdk')"))
 	FDirectoryPath JavaPath;
+#endif
 
 #if WITH_EDITOR
 	// UObject interface
diff --git a/Engine/Source/Developer/Android/AndroidPlatformEditor/Private/AndroidSDKSettingsCustomization.cpp b/Engine/Source/Developer/Android/AndroidPlatformEditor/Private/AndroidSDKSettingsCustomization.cpp
--- a/Engine/Source/Developer/Android/AndroidPlatformEditor/Private/AndroidSDKSettingsCustomization.cpp
+++ b/Engine/Source/Developer/Android/AndroidPlatformEditor/Private/AndroidSDKSettingsCustomization.cpp
@@ -88,7 +88,7 @@ void FAndroidSDKSettingsCustomization::SetupSDKPaths()
 		}
 		
 	}
-
+#if PLATFORM_MAC == 0
 	if (settings->JavaPath.Path.IsEmpty())
 	{
 		TCHAR AndroidJavaPath[256];
@@ -100,7 +100,7 @@ void FAndroidSDKSettingsCustomization::SetupSDKPaths()
 		}
 
 	}
-	
+#endif	
 	if (changed)
 	{
 		settings->Modify();
