- Added some in-editor documentation/tips to clarify some of the project settings [UE-8127]
- Only allow ES31+AEP if the project settings allows it. Told user to restart editor if they change the ES31+AEP setting to fix the Launch On menu. (This is only until 4.8) [UE-6585]
- Noted that ES31 is only for Tegra K1, in the name [UE-6604]
- Hide JAVA_HOME settings on MacOS, as it would be confusing to try to find it

diff --git a/Engine/Source/Developer/Android/AndroidPlatformEditor/Private/AndroidSDKSettings.h b/Engine/Source/Developer/Android/AndroidPlatformEditor/Private/AndroidSDKSettings.h
--- a/Engine/Source/Developer/Android/AndroidPlatformEditor/Private/AndroidSDKSettings.h
+++ b/Engine/Source/Developer/Android/AndroidPlatformEditor/Private/AndroidSDKSettings.h
@@ -16,18 +16,22 @@ class ANDROIDPLATFORMEDITOR_API UAndroidSDKSettings : public UObject
 public:
 	GENERATED_UCLASS_BODY()
 
-	// Location on disk to find the Android SDK (defaults to ANDROID_HOME environment variable if set)
+	// Location on disk of the Android SDK (defaults to ANDROID_HOME environment variable if set)
 	UPROPERTY(GlobalConfig, EditAnywhere, Category = SDKConfig, Meta = (DisplayName = "Location of Android SDK (the directory usually contains 'android-sdk-')"))
 	FDirectoryPath SDKPath;
 
-	// Location on disk to find the Android NDK (defaults to NDKROOT environment variable if set)
+	// Location on disk of the Android NDK (defaults to NDKROOT environment variable if set)
 	UPROPERTY(GlobalConfig, EditAnywhere, Category = SDKConfig, Meta = (DisplayName = "Location of Android NDK (the directory usually contains 'android-ndk-')"))
 	FDirectoryPath NDKPath;
 
-	// Location on disk to find the ANT tool
+	// Location on disk of the ANT tool
 	UPROPERTY(GlobalConfig, EditAnywhere, Category = SDKConfig, Meta = (DisplayName = "Location of ANT (the directory usually contains 'apache-ant-')"))
 	FDirectoryPath ANTPath;
 
+	// Location on disk of Java (this it not required on the Mac)
+	UPROPERTY(GlobalConfig, EditAnywhere, Category = SDKConfig, Meta = (DisplayName = "Location of JAVA (the directory usually contains 'jdk')"))
+	FDirectoryPath JavaPath;
+	
 	// Which SDK to package and compile Java with (a specific version or (without quotes) 'latest' for latest version on disk, or 'matchndk' to match the NDK API Level)
 	UPROPERTY(GlobalConfig, EditAnywhere, Category = SDKConfig, Meta = (DisplayName = "SDK API Level (specific version, 'latest', or 'matchndk' - see tooltip)"))
 	FString SDKAPILevel;
@@ -36,10 +40,6 @@ public:
 	UPROPERTY(GlobalConfig, EditAnywhere, Category = SDKConfig, Meta = (DisplayName = "NDK API Level (specific version or 'latest' - see tooltip)"))
 	FString NDKAPILevel;
 
-#if PLATFORM_MAC == 0
-	UPROPERTY(GlobalConfig, EditAnywhere, Category = SDKConfig, Meta = (DisplayName = "Location of JAVA (the directory usually contains 'jdk')"))
-	FDirectoryPath JavaPath;
-#endif
 
 #if WITH_EDITOR
 	// UObject interface
diff --git a/Engine/Source/Developer/Android/AndroidPlatformEditor/Private/AndroidSDKSettingsCustomization.cpp b/Engine/Source/Developer/Android/AndroidPlatformEditor/Private/AndroidSDKSettingsCustomization.cpp
--- a/Engine/Source/Developer/Android/AndroidPlatformEditor/Private/AndroidSDKSettingsCustomization.cpp
+++ b/Engine/Source/Developer/Android/AndroidPlatformEditor/Private/AndroidSDKSettingsCustomization.cpp
@@ -44,6 +44,15 @@ void FAndroidSDKSettingsCustomization::CustomizeDetails(IDetailLayoutBuilder& De
 void FAndroidSDKSettingsCustomization::BuildSDKPathSection(IDetailLayoutBuilder& DetailLayout)
 {
 	SetupSDKPaths();
+
+#if PLATFORM_MAC
+	IDetailCategoryBuilder& SDKConfigCategory = DetailLayout.EditCategory(TEXT("SDKConfig"));
+
+	// hide the property on Mac only
+	TSharedRef<IPropertyHandle> JavaPathProperty = DetailLayout.GetProperty(GET_MEMBER_NAME_CHECKED(UAndroidSDKSettings, JavaPath));
+	SDKConfigCategory.AddProperty(JavaPathProperty)
+		.Visibility(EVisibility::Hidden);
+#endif
 }
 
 
diff --git a/Engine/Source/Developer/Android/AndroidPlatformEditor/Private/AndroidTargetSettingsCustomization.cpp b/Engine/Source/Developer/Android/AndroidPlatformEditor/Private/AndroidTargetSettingsCustomization.cpp
--- a/Engine/Source/Developer/Android/AndroidPlatformEditor/Private/AndroidTargetSettingsCustomization.cpp
+++ b/Engine/Source/Developer/Android/AndroidPlatformEditor/Private/AndroidTargetSettingsCustomization.cpp
@@ -56,6 +56,17 @@ void FAndroidTargetSettingsCustomization::CustomizeDetails(IDetailLayoutBuilder&
 	BuildIconSection(DetailLayout);
 }
 
+static void OnBrowserLinkClicked(const FSlateHyperlinkRun::FMetadata& Metadata)
+{
+	const FString* URL = Metadata.Find(TEXT("href"));
+	
+	if(URL)
+	{
+		FPlatformProcess::LaunchURL(**URL, nullptr, nullptr);
+	}
+}
+
+
 void FAndroidTargetSettingsCustomization::BuildAppManifestSection(IDetailLayoutBuilder& DetailLayout)
 {
 	// Cache some categories
@@ -75,6 +86,27 @@ void FAndroidTargetSettingsCustomization::BuildAppManifestSection(IDetailLayoutB
 			PlatformSetupMessage
 		];
 
+	APKPackagingCategory.AddCustomRow(LOCTEXT("UpgradeInfo", "Upgrade Info"), false)
+		.WholeRowWidget
+		[
+			SNew(SBorder)
+			.Padding(1)
+			[
+				SNew(SHorizontalBox)
+				+ SHorizontalBox::Slot()
+				.Padding(FMargin(10, 10, 10, 10))
+				.FillWidth(1.0f)
+				[
+					SNew(SRichTextBlock)
+					.Text(LOCTEXT("UpgradeInfoMessage", "<RichTextBlock.TextHighlight>Note to users from 4.6 or earlier</>: We now <RichTextBlock.TextHighlight>GENERATE</> an AndroidManifest.xml when building, so if you have customized your .xml file, you will need to put all of your changes into the below settings. Note that we don't touch your AndroidManifest.xml that is in your project directory.\nAdditionally, we no longer use SigningConfig.xml, the settings are now set in the Distribution Signing section.\n\nThere is currently no .obb file downloader support in the engine, so if you don't package your data into your .apk (see the below setting and its tooltip about 50MB limit), device is not guaranteed to have the .obb file downloaded in all cases. Until Unreal Engine v4.8, there won't be a way for your app to download the .obb file from the Google Play Store. See <a id=\"browser\" href=\"http://developer.android.com/google/play/expansion-files.html#Downloading\" style=\"HoverOnlyHyperlink\">http://developer.android.com/google/play/expansion-files.html</> for more information."))
+					.TextStyle(FEditorStyle::Get(), "MessageLog")
+					.DecoratorStyleSet(&FEditorStyle::Get())
+					.AutoWrapText(true)
+					+ SRichTextBlock::HyperlinkDecorator(TEXT("browser"), FSlateHyperlinkRun::FOnClick::CreateStatic(&OnBrowserLinkClicked))
+				]
+			]
+		];
+	
 	APKPackagingCategory.AddCustomRow(LOCTEXT("BuildFolderLabel", "Build Folder"), false)
 		.IsEnabled(SetupForPlatformAttribute)
 		.NameContent()
@@ -171,7 +203,8 @@ void FAndroidTargetSettingsCustomization::BuildAppManifestSection(IDetailLayoutB
 	SETUP_NONROCKET_PROP(bBuildForArmV7, BuildCategory, LOCTEXT("BuildForArmV7ToolTip", "Enable ArmV7 CPU architecture support? (this will be used if all CPU architecture types are unchecked)"));
 	SETUP_NONROCKET_PROP(bBuildForX86, BuildCategory, LOCTEXT("BuildForX86ToolTip", "Enable X86 CPU architecture support?"));
 	SETUP_NONROCKET_PROP(bBuildForES2, BuildCategory, LOCTEXT("BuildForES2ToolTip", "Enable OpenGL ES2 rendering support? (this will be used if rendering types are unchecked)"));
-	SETUP_NONROCKET_PROP(bBuildForES31, BuildCategory, LOCTEXT("BuildForES31ToolTip", "Enable OpenGL ES31 + AEP (Android Extension Pack) rendering support?"));
+	SETUP_NONROCKET_PROP(bBuildForES31, BuildCategory, LOCTEXT("BuildForES31ToolTip", "Enable OpenGL ES31 + AEP (Android Extension Pack) rendering support? Currently only Tegra K1 supports this, as it will force DXT textures (In 4.8 3.1+AEP will work with all texture formats).\nIf you use the Launch On feature (in the main toolbar), when you change this setting, you need to restart the editor to make sure it will launch with the proper 3.1+AEP support!"));
+	
 	// @todo android fat binary: Put back in when we expose those
 //	SETUP_NONROCKET_PROP(bSplitIntoSeparateApks, BuildCategory, LOCTEXT("SplitIntoSeparateAPKsToolTip", "If checked, CPU architectures and rendering types will be split into separate .apk files"));
 }
diff --git a/Engine/Source/Developer/Android/Android_ES31TargetPlatform/Private/Android_ES31TargetPlatformModule.cpp b/Engine/Source/Developer/Android/Android_ES31TargetPlatform/Private/Android_ES31TargetPlatformModule.cpp
--- a/Engine/Source/Developer/Android/Android_ES31TargetPlatform/Private/Android_ES31TargetPlatformModule.cpp
+++ b/Engine/Source/Developer/Android/Android_ES31TargetPlatform/Private/Android_ES31TargetPlatformModule.cpp
@@ -22,7 +22,7 @@ class FAndroid_ES31TargetPlatform
 
 	virtual FText DisplayName( ) const override
 	{
-		return LOCTEXT("Android_ES31", "Android (ES 3.1 + AEP)");
+		return LOCTEXT("Android_ES31", "Android (ES 3.1 + AEP, TegraK1)");
 	}
 
 	virtual FString PlatformName() const override
@@ -43,13 +43,23 @@ class FAndroid_ES31TargetPlatform
 
 	virtual bool SupportedByExtensionsString( const FString& ExtensionsString, const int GLESVersion ) const override
 	{
-		// look for the AEP extension
-		return ExtensionsString.Contains(TEXT("GL_ANDROID_extension_pack_es31a")) && ExtensionsString.Contains(TEXT("GL_EXT_color_buffer_half_float"));
+		// look for the AEP extension, and make sure the project supports it (if the user changes the setting, they will need to restart
+		// the editor, a tooltip has been included to make this clear)
+		bool bProjectSupportsES31;
+		if (GConfig->GetBool(TEXT("/Script/AndroidRuntimeSettings.AndroidRuntimeSettings"), TEXT("bBuildForES31"), bProjectSupportsES31, GEngineIni))
+		{
+			if (bProjectSupportsES31)
+			{
+				return ExtensionsString.Contains(TEXT("GL_ANDROID_extension_pack_es31a")) && ExtensionsString.Contains(TEXT("GL_EXT_color_buffer_half_float"));
+			}
+		}
+
+		return false;
 	}
 
 	virtual FText GetVariantDisplayName() const override
 	{
-		return LOCTEXT("Android_31_ShortName", "ES 3.1 + AEP");
+		return LOCTEXT("Android_31_ShortName", "ES 3.1 + AEP, TegraK1");
 	}
 
 	virtual float GetVariantPriority() const override
diff --git a/Engine/Source/Developer/iOS/IOSPlatformEditor/Private/IOSTargetSettingsCustomization.cpp b/Engine/Source/Developer/iOS/IOSPlatformEditor/Private/IOSTargetSettingsCustomization.cpp
--- a/Engine/Source/Developer/iOS/IOSPlatformEditor/Private/IOSTargetSettingsCustomization.cpp
+++ b/Engine/Source/Developer/iOS/IOSPlatformEditor/Private/IOSTargetSettingsCustomization.cpp
@@ -230,52 +230,27 @@ void FIOSTargetSettingsCustomization::BuildPListSection(IDetailLayoutBuilder& De
 						]
 					]
 			];
-
-	AppManifestCategory.AddCustomRow(LOCTEXT("Warning", "Warning"), false)
-		.WholeRowWidget
-		[
-			PlatformSetupMessage
-		];
-
-	AppManifestCategory.AddCustomRow(LOCTEXT("InfoPlistHyperlink", "Info.plist Hyperlink"), false)
-		.WholeRowWidget
-		[
-			SNew(SBox)
-			.HAlign(HAlign_Center)
-			[
-				SNew(SHyperlinkLaunchURL, TEXT("https://developer.apple.com/library/ios/documentation/general/Reference/InfoPlistKeyReference/Articles/AboutInformationPropertyListFiles.html"))
-				.Text(LOCTEXT("ApplePlistPage", "About Information Property List Files"))
-				.ToolTipText(LOCTEXT("ApplePlistPageTooltip", "Opens a page that discusses Info.plist"))
-			]
-		];
-
-
-	AppManifestCategory.AddCustomRow(LOCTEXT("InfoPlist", "Info.plist"), false)
-		.IsEnabled(SetupForPlatformAttribute)
-		.NameContent()
+	
+	BundleCategory.AddCustomRow(LOCTEXT("UpgradeInfo", "Upgrade Info"), false)
+	.WholeRowWidget
+	[
+		SNew(SBorder)
+		.Padding(1)
 		[
 			SNew(SHorizontalBox)
 			+ SHorizontalBox::Slot()
-			.Padding(FMargin(0, 1, 0, 1))
+			.Padding(FMargin(10, 10, 10, 10))
 			.FillWidth(1.0f)
 			[
-				SNew(STextBlock)
-				.Text(LOCTEXT("PlistLabel", "Info.plist"))
-				.Font(DetailLayout.GetDetailFont())
-			]
-		]
-		.ValueContent()
-		[
-			SNew(SHorizontalBox)
-			+SHorizontalBox::Slot()
-			.AutoWidth()
-			[
-				SNew(SButton)
-				.Text(LOCTEXT("OpenPlistFolderButton", "Open PList Folder"))
-				.ToolTipText(LOCTEXT("OpenPlistFolderButton_Tooltip", "Opens the folder containing the plist for the current project in Explorer or Finder"))
-				.OnClicked(this, &FIOSTargetSettingsCustomization::OpenPlistFolder)
-			]
-		];
+				SNew(SRichTextBlock)
+				.Text(LOCTEXT("IOSUpgradeInfoMessage", "<RichTextBlock.TextHighlight>Note to users from 4.6 or earlier</>: We now <RichTextBlock.TextHighlight>GENERATE</> an Info.plist when building, so if you have customized your .plist file, you will need to put all of your changes into the below settings. Note that we don't touch the .plist file that is in your project directory, so you can use it as reference."))
+				.TextStyle(FEditorStyle::Get(), "MessageLog")
+				.DecoratorStyleSet(&FEditorStyle::Get())
+				.AutoWrapText(true)
+				// + SRichTextBlock::HyperlinkDecorator(TEXT("browser"), FSlateHyperlinkRun::FOnClick::CreateStatic(&OnBrowserLinkClicked))
+			 ]
+		 ]
+	 ];
 
 	// Show properties that are gated by the plist being present and writable
 	FSimpleDelegate PlistModifiedDelegate = FSimpleDelegate::CreateRaw(this, &FIOSTargetSettingsCustomization::OnPlistPropertyModified);
