[INTEGRATE] Change 2423627 by [REDACTED_PERSON][REDACTED_HANDLE] on 2015/01/29 07:23:55

	UE-8509 Android SDK Paths being saved even if the user doesn't set them
	- removed initial pull of data on start up

diff --git a/Engine/Source/Developer/Android/AndroidPlatformEditor/Private/AndroidPlatformEditorModule.cpp b/Engine/Source/Developer/Android/AndroidPlatformEditor/Private/AndroidPlatformEditorModule.cpp
--- a/Engine/Source/Developer/Android/AndroidPlatformEditor/Private/AndroidPlatformEditorModule.cpp
+++ b/Engine/Source/Developer/Android/AndroidPlatformEditor/Private/AndroidPlatformEditorModule.cpp
@@ -59,8 +59,7 @@ class FAndroidPlatformEditorModule
 		settings->SetTargetModule(&TargetPlatformManagerModule);
 		auto &AndroidDeviceDetection = FModuleManager::LoadModuleChecked<IAndroidDeviceDetectionModule>("AndroidDeviceDetection");
 		settings->SetDeviceDetection(AndroidDeviceDetection.GetAndroidDeviceDetection());
-		settings->SetupInitialTargetPaths();
-
+		settings->UpdateTargetModulePaths();
 	}
 
 	virtual void ShutdownModule() override
diff --git a/Engine/Source/Developer/Android/AndroidPlatformEditor/Private/AndroidSDKSettings.cpp b/Engine/Source/Developer/Android/AndroidPlatformEditor/Private/AndroidSDKSettings.cpp
--- a/Engine/Source/Developer/Android/AndroidPlatformEditor/Private/AndroidSDKSettings.cpp
+++ b/Engine/Source/Developer/Android/AndroidPlatformEditor/Private/AndroidSDKSettings.cpp
@@ -17,7 +17,7 @@ UAndroidSDKSettings::UAndroidSDKSettings(const FObjectInitializer& ObjectInitial
 void UAndroidSDKSettings::PostEditChangeProperty(struct FPropertyChangedEvent& PropertyChangedEvent)
 {
 	Super::PostEditChangeProperty(PropertyChangedEvent);
-	UpdateTargetModulePaths(true);	
+	UpdateTargetModulePaths();	
 }
 
 void UAndroidSDKSettings::SetTargetModule(ITargetPlatformManagerModule * TargetManagerModule)
@@ -30,115 +30,26 @@ void UAndroidSDKSettings::SetDeviceDetection(IAndroidDeviceDetection * AndroidDe
 	this->AndroidDeviceDetection = AndroidDeviceDetection;
 }
 
-void UAndroidSDKSettings::SetupInitialTargetPaths()
-{
-	// Check each path, if we don't have it then try to look up the values from a known location instead
-	if (SDKPath.Path.IsEmpty())
-	{
-		TCHAR AndroidSDKPath[256];
-		FPlatformMisc::GetEnvironmentVariable(TEXT("ANDROID_HOME"), AndroidSDKPath, ARRAY_COUNT(AndroidSDKPath));
-		if (AndroidSDKPath[0])
-		{
-			SDKPath.Path = AndroidSDKPath;
-		}
-	}
-
-	if (NDKPath.Path.IsEmpty())
-	{
-		TCHAR AndroidNDKPath[256];
-		FPlatformMisc::GetEnvironmentVariable(TEXT("NDKROOT"), AndroidNDKPath, ARRAY_COUNT(AndroidNDKPath));
-		if (AndroidNDKPath[0])
-		{
-			NDKPath.Path = AndroidNDKPath;
-		}
-	}
-
-	if (ANTPath.Path.IsEmpty())
-	{
-		TCHAR AndroidANTPath[256];
-		FPlatformMisc::GetEnvironmentVariable(TEXT("ANT_HOME"), AndroidANTPath, ARRAY_COUNT(AndroidANTPath));
-		if (AndroidANTPath[0])
-		{
-			ANTPath.Path = AndroidANTPath;
-		}
-	}
-
-#if PLATFORM_MAC == 0
-	if (JavaPath.Path.IsEmpty())
-	{
-		TCHAR AndroidJavaPath[256];
-		FPlatformMisc::GetEnvironmentVariable(TEXT("JAVA_HOME"), AndroidJavaPath, ARRAY_COUNT(AndroidJavaPath));
-		if (AndroidJavaPath[0])
-		{
-			JavaPath.Path = AndroidJavaPath;
-		}
-	}
-#endif
-
-#if PLATFORM_MAC
-	// On a Mac, if we still don't have the values then check the .bash_profile file for them
-	if (SDKPath.Path.IsEmpty() || NDKPath.Path.IsEmpty() || ANTPath.Path.IsEmpty())
-	{
-		FArchive* FileReader = IFileManager::Get().CreateFileReader(*FString([@"~/.bash_profile" stringByExpandingTildeInPath]));
-		if (FileReader)
-		{
-			const int64 FileSize = FileReader->TotalSize();
-			ANSICHAR* AnsiContents = (ANSICHAR*)FMemory::Malloc(FileSize);
-			FileReader->Serialize(AnsiContents, FileSize);
-			FileReader->Close();
-			delete FileReader;
-
-			TArray<FString> Lines;
-			FString(ANSI_TO_TCHAR(AnsiContents)).ParseIntoArrayLines(&Lines);
-			FMemory::Free(AnsiContents);
-
-			for (int32 Index = 0; Index < Lines.Num(); Index++)
-			{
-				if (SDKPath.Path.IsEmpty() && Lines[Index].StartsWith(TEXT("export ANDROID_HOME=")))
-				{
-					FString Directory;
-					Lines[Index].Split(TEXT("="), NULL, &Directory);
-					SDKPath.Path = Directory.Replace(TEXT("\""), TEXT(""));
-				}
-				else if (ANTPath.Path.IsEmpty() && Lines[Index].StartsWith(TEXT("export ANT_HOME=")))
-				{
-					FString Directory;
-					Lines[Index].Split(TEXT("="), NULL, &Directory);
-					ANTPath.Path = Directory.Replace(TEXT("\""), TEXT(""));
-				}
-				else if (NDKPath.Path.IsEmpty() && Lines[Index].StartsWith(TEXT("export NDKROOT=")))
-				{
-					FString Directory;
-					Lines[Index].Split(TEXT("="), NULL, &Directory);
-					NDKPath.Path = Directory.Replace(TEXT("\""), TEXT(""));
-				}
-			}
-		}
-	}
-
-#endif
-	UpdateTargetModulePaths(false);
-}
-void UAndroidSDKSettings::UpdateTargetModulePaths(bool bForceUpdate)
+void UAndroidSDKSettings::UpdateTargetModulePaths()
 {
 	TArray<FString> Keys;
 	TArray<FString> Values;
 	
-	if (bForceUpdate || !SDKPath.Path.IsEmpty())
+	if (!SDKPath.Path.IsEmpty())
 	{
 		FPaths::NormalizeFilename(SDKPath.Path);
 		Keys.Add(TEXT("ANDROID_HOME"));
 		Values.Add(SDKPath.Path);
 	}
 	
-	if (bForceUpdate || !NDKPath.Path.IsEmpty())
+	if (!NDKPath.Path.IsEmpty())
 	{
 		FPaths::NormalizeFilename(NDKPath.Path);
 		Keys.Add(TEXT("NDKROOT"));
 		Values.Add(NDKPath.Path);
 	}
 	
-	if (bForceUpdate || !ANTPath.Path.IsEmpty())
+	if (!ANTPath.Path.IsEmpty())
 	{
 		FPaths::NormalizeFilename(ANTPath.Path);
 		Keys.Add(TEXT("ANT_HOME"));
@@ -146,7 +57,7 @@ void UAndroidSDKSettings::UpdateTargetModulePaths(bool bForceUpdate)
 	}
 
 #if PLATFORM_MAC == 0
-	if (bForceUpdate || !JavaPath.Path.IsEmpty())
+	if (!JavaPath.Path.IsEmpty())
 	{
 		FPaths::NormalizeFilename(JavaPath.Path);
 		Keys.Add(TEXT("JAVA_HOME"));
diff --git a/Engine/Source/Developer/Android/AndroidPlatformEditor/Private/AndroidSDKSettings.h b/Engine/Source/Developer/Android/AndroidPlatformEditor/Private/AndroidSDKSettings.h
--- a/Engine/Source/Developer/Android/AndroidPlatformEditor/Private/AndroidSDKSettings.h
+++ b/Engine/Source/Developer/Android/AndroidPlatformEditor/Private/AndroidSDKSettings.h
@@ -47,8 +47,7 @@ public:
 	// End of UObject interface
 	void SetTargetModule(ITargetPlatformManagerModule * TargetManagerModule);
 	void SetDeviceDetection(IAndroidDeviceDetection * AndroidDeviceDetection);
-	void SetupInitialTargetPaths();
-	void UpdateTargetModulePaths(bool bForceUpdate);
+	void UpdateTargetModulePaths();
 	ITargetPlatformManagerModule * TargetManagerModule;
 	IAndroidDeviceDetection * AndroidDeviceDetection;
 #endif
diff --git a/Engine/Source/Developer/Android/AndroidPlatformEditor/Private/AndroidSDKSettingsCustomization.cpp b/Engine/Source/Developer/Android/AndroidPlatformEditor/Private/AndroidSDKSettingsCustomization.cpp
--- a/Engine/Source/Developer/Android/AndroidPlatformEditor/Private/AndroidSDKSettingsCustomization.cpp
+++ b/Engine/Source/Developer/Android/AndroidPlatformEditor/Private/AndroidSDKSettingsCustomization.cpp
@@ -31,7 +31,6 @@ TSharedRef<IDetailCustomization> FAndroidSDKSettingsCustomization::MakeInstance(
 FAndroidSDKSettingsCustomization::FAndroidSDKSettingsCustomization()
 {
 	TargetPlatformManagerModule = &FModuleManager::LoadModuleChecked<ITargetPlatformManagerModule>("TargetPlatform");
-	SetupSDKPaths();
 }
 
 void FAndroidSDKSettingsCustomization::CustomizeDetails(IDetailLayoutBuilder& DetailLayout)
@@ -43,8 +42,6 @@ void FAndroidSDKSettingsCustomization::CustomizeDetails(IDetailLayoutBuilder& De
 
 void FAndroidSDKSettingsCustomization::BuildSDKPathSection(IDetailLayoutBuilder& DetailLayout)
 {
-	SetupSDKPaths();
-
 #if PLATFORM_MAC
 	IDetailCategoryBuilder& SDKConfigCategory = DetailLayout.EditCategory(TEXT("SDKConfig"));
 
@@ -55,71 +52,6 @@ void FAndroidSDKSettingsCustomization::BuildSDKPathSection(IDetailLayoutBuilder&
 #endif
 }
 
-
-void FAndroidSDKSettingsCustomization::SetupSDKPaths()
-{
-	// Try and query the env vars
-	UAndroidSDKSettings * settings = GetMutableDefault<UAndroidSDKSettings>();
-	bool changed = false;
-	
-	if (settings->SDKPath.Path.IsEmpty())
-	{
-		TCHAR AndroidSDKPath[256];
-		FPlatformMisc::GetEnvironmentVariable(TEXT("ANDROID_HOME"), AndroidSDKPath, ARRAY_COUNT(AndroidSDKPath));
-		if (AndroidSDKPath[0])
-		{
-			settings->SDKPath.Path = AndroidSDKPath;
-			changed |= true;
-		}
-				
-	}	
-
-	if (settings->NDKPath.Path.IsEmpty())
-	{
-		TCHAR AndroidNDKPath[256];
-		FPlatformMisc::GetEnvironmentVariable(TEXT("NDKROOT"), AndroidNDKPath, ARRAY_COUNT(AndroidNDKPath));
-		if (AndroidNDKPath[0])
-		{
-			settings->NDKPath.Path = AndroidNDKPath;
-			changed |= true;
-		}
-		
-	}	
-
-	if (settings->ANTPath.Path.IsEmpty())
-	{
-		TCHAR AndroidANTPath[256];
-		FPlatformMisc::GetEnvironmentVariable(TEXT("ANT_HOME"), AndroidANTPath, ARRAY_COUNT(AndroidANTPath));
-		if (AndroidANTPath[0])
-		{
-			settings->ANTPath.Path = AndroidANTPath;
-			changed |= true;
-		}
-		
-	}
-#if PLATFORM_MAC == 0
-	if (settings->JavaPath.Path.IsEmpty())
-	{
-		TCHAR AndroidJavaPath[256];
-		FPlatformMisc::GetEnvironmentVariable(TEXT("JAVA_HOME"), AndroidJavaPath, ARRAY_COUNT(AndroidJavaPath));
-		if (AndroidJavaPath[0])
-		{
-			settings->JavaPath.Path = AndroidJavaPath;
-			changed |= true;
-		}
-
-	}
-#endif	
-	if (changed)
-	{
-		settings->Modify();
-		settings->SaveConfig();
-	}
-
-	settings->UpdateTargetModulePaths(false);
-
-}
-
 //////////////////////////////////////////////////////////////////////////
 
 #undef LOCTEXT_NAMESPACE
diff --git a/Engine/Source/Developer/Android/AndroidPlatformEditor/Private/AndroidSDKSettingsCustomization.h b/Engine/Source/Developer/Android/AndroidPlatformEditor/Private/AndroidSDKSettingsCustomization.h
--- a/Engine/Source/Developer/Android/AndroidPlatformEditor/Private/AndroidSDKSettingsCustomization.h
+++ b/Engine/Source/Developer/Android/AndroidPlatformEditor/Private/AndroidSDKSettingsCustomization.h
@@ -23,7 +23,6 @@ private:
 	FAndroidSDKSettingsCustomization();
 
 	void BuildSDKPathSection(IDetailLayoutBuilder& DetailLayout);
-	void SetupSDKPaths();
 
 private:
 	IDetailLayoutBuilder* SavedLayoutBuilder;
