- Fixed confusing Android SDK tooltips [UE-9155]

diff --git a/Engine/Source/Developer/Android/AndroidPlatformEditor/Private/AndroidSDKSettings.h b/Engine/Source/Developer/Android/AndroidPlatformEditor/Private/AndroidSDKSettings.h
--- a/Engine/Source/Developer/Android/AndroidPlatformEditor/Private/AndroidSDKSettings.h
+++ b/Engine/Source/Developer/Android/AndroidPlatformEditor/Private/AndroidSDKSettings.h
@@ -16,19 +16,19 @@ class ANDROIDPLATFORMEDITOR_API UAndroidSDKSettings : public UObject
 public:
 	GENERATED_UCLASS_BODY()
 
-	// Location on disk of the Android SDK (defaults to ANDROID_HOME environment variable if set)
+	// Location on disk of the Android SDK (falls back to ANDROID_HOME environment variable if this is left blank)
 	UPROPERTY(GlobalConfig, EditAnywhere, Category = SDKConfig, Meta = (DisplayName = "Location of Android SDK (the directory usually contains 'android-sdk-')"))
 	FDirectoryPath SDKPath;
 
-	// Location on disk of the Android NDK (defaults to NDKROOT environment variable if set)
+	// Location on disk of the Android NDK (falls back to NDKROOT environment variable if this is left blank)
 	UPROPERTY(GlobalConfig, EditAnywhere, Category = SDKConfig, Meta = (DisplayName = "Location of Android NDK (the directory usually contains 'android-ndk-')"))
 	FDirectoryPath NDKPath;
 
-	// Location on disk of the ANT tool
+	// Location on disk of the ANT tool (falls back to ANT_HOME environment variable if this is left blank)
 	UPROPERTY(GlobalConfig, EditAnywhere, Category = SDKConfig, Meta = (DisplayName = "Location of ANT (the directory usually contains 'apache-ant-')"))
 	FDirectoryPath ANTPath;
 
-	// Location on disk of Java (this it not required on the Mac)
+	// Location on disk of Java (falls back to JAVA_HOME environment variable if this is left blank)
 	UPROPERTY(GlobalConfig, EditAnywhere, Category = SDKConfig, Meta = (DisplayName = "Location of JAVA (the directory usually contains 'jdk')"))
 	FDirectoryPath JavaPath;
 	
