[INTEGRATE] Change 2456485 by [REDACTED_PERSON][REDACTED_HANDLE] on 2015/02/23 13:34:43

	fix for ANDROID_HOME not being set on Mac by default
	UE-10257, UE-10277

diff --git a/Engine/Source/Developer/Android/AndroidDeviceDetection/Private/AndroidDeviceDetectionModule.cpp b/Engine/Source/Developer/Android/AndroidDeviceDetection/Private/AndroidDeviceDetectionModule.cpp
--- a/Engine/Source/Developer/Android/AndroidDeviceDetection/Private/AndroidDeviceDetectionModule.cpp
+++ b/Engine/Source/Developer/Android/AndroidDeviceDetection/Private/AndroidDeviceDetectionModule.cpp
@@ -313,6 +313,38 @@ public:
 
 		FString ADBPath;
 		
+#if PLATFORM_MAC
+		if (AndroidDirectory[0] == 0)
+		{
+			// didn't find ANDROID_HOME, so parse the .bash_profile file on MAC
+			FArchive* FileReader = IFileManager::Get().CreateFileReader(*FString([@"~/.bash_profile" stringByExpandingTildeInPath]));
+			if (FileReader)
+			{
+				const int64 FileSize = FileReader->TotalSize();
+				ANSICHAR* AnsiContents = (ANSICHAR*)FMemory::Malloc(FileSize);
+				FileReader->Serialize(AnsiContents, FileSize);
+				FileReader->Close();
+				delete FileReader;
+
+				TArray<FString> Lines;
+				FString(ANSI_TO_TCHAR(AnsiContents)).ParseIntoArrayLines(&Lines);
+				FMemory::Free(AnsiContents);
+
+				for (int32 Index = 0; Index < Lines.Num(); Index++)
+				{
+					if (AndroidDirectory[0] == 0 && Lines[Index].StartsWith(TEXT("export ANDROID_HOME=")))
+					{
+						FString Directory;
+						Lines[Index].Split(TEXT("="), NULL, &Directory);
+						Directory = Directory.Replace(TEXT("\""), TEXT(""));
+						FCString::Strcpy(AndroidDirectory, *Directory);
+						setenv("ANDROID_HOME", TCHAR_TO_ANSI(AndroidDirectory), 1);
+					}
+				}
+			}
+		}
+#endif
+
 		if (AndroidDirectory[0] != 0)
 		{
 #if PLATFORM_WINDOWS
