PR #894: wrong parsing of the return of adb  cmd: device (Contributed by [REDACTED_PERSON]) with update to accept either tab or space
https://github.com/EpicGames/UnrealEngine/pull/894
#ue4
#android

diff --git a/Engine/Source/Developer/Android/AndroidDeviceDetection/Private/AndroidDeviceDetectionModule.cpp b/Engine/Source/Developer/Android/AndroidDeviceDetection/Private/AndroidDeviceDetectionModule.cpp
--- a/Engine/Source/Developer/Android/AndroidDeviceDetection/Private/AndroidDeviceDetectionModule.cpp
+++ b/Engine/Source/Developer/Android/AndroidDeviceDetection/Private/AndroidDeviceDetectionModule.cpp
@@ -135,9 +135,13 @@ private:
 			// grab the device serial number
 			int32 TabIndex;
 
-			if (!DeviceString.FindChar(TCHAR(' '), TabIndex))
+			// use either tab or space as separator
+			if (!DeviceString.FindChar(TCHAR('\t'), TabIndex))
 			{
-				continue;
+				if (!DeviceString.FindChar(TCHAR(' '), TabIndex))
+				{
+					continue;
+				}
 			}
 
 			FAndroidDeviceInfo NewDeviceInfo;
