Android ASTC for ES2

diff --git a/Engine/Source/Developer/Android/AndroidTargetPlatform/Private/AndroidTargetPlatform.h b/Engine/Source/Developer/Android/AndroidTargetPlatform/Private/AndroidTargetPlatform.h
--- a/Engine/Source/Developer/Android/AndroidTargetPlatform/Private/AndroidTargetPlatform.h
+++ b/Engine/Source/Developer/Android/AndroidTargetPlatform/Private/AndroidTargetPlatform.h
@@ -33,6 +33,12 @@ namespace AndroidTexFormat
 	static FName NameETC2_RGB(TEXT("ETC2_RGB"));
 	static FName NameETC2_RGBA(TEXT("ETC2_RGBA"));
 	static FName NameAutoETC2(TEXT("AutoETC2"));
+	static FName NameASTC_4x4(TEXT("ASTC_4x4"));
+	static FName NameASTC_6x6(TEXT("ASTC_6x6"));
+	static FName NameASTC_8x8(TEXT("ASTC_8x8"));
+	static FName NameASTC_10x10(TEXT("ASTC_10x10"));
+	static FName NameASTC_12x12(TEXT("ASTC_12x12"));
+	static FName NameAutoASTC(TEXT("AutoASTC"));
 
 	// Uncompressed Texture Formats
 	static FName NameBGRA8(TEXT("BGRA8"));
@@ -64,7 +70,7 @@ public:
 public:
 
 	/**
-	 * Gets the name of the Android platform variant, i.e. ATC, DXT or PVRTC.
+	 * Gets the name of the Android platform variant, i.e. ATC, DXT, PVRTC, etc.
 	 *
 	 * @param Variant name.
 	 */
@@ -172,13 +178,16 @@ protected:
 		return true;
 	}
 
+#if WITH_ENGINE
+	// Holds the Engine INI settings (for quick access).
+	FConfigFile EngineSettings;
+#endif //WITH_ENGINE
+
 private:
 
 	// Handles when the ticker fires.
 	bool HandleTicker( float DeltaTime );
 
-private:
-
 	// Holds a map of valid devices.
 	TMap<FString, FAndroidTargetDevicePtr> Devices;
 
@@ -192,9 +201,6 @@ private:
 	IAndroidDeviceDetection* DeviceDetection;
 
 #if WITH_ENGINE
-	// Holds the Engine INI settings (for quick access).
-	FConfigFile EngineSettings;
-
 	// Holds a cache of the target LOD settings.
 	FTextureLODSettings TextureLODSettings;
 
@@ -204,8 +210,6 @@ private:
 	ITargetDevicePtr DefaultDevice;
 #endif //WITH_ENGINE
 
-private:
-
 	// Holds an event delegate that is executed when a new target device has been discovered.
 	FOnTargetDeviceDiscovered DeviceDiscoveredEvent;
 
diff --git a/Engine/Source/Developer/Android/Android_ASTCTargetPlatform/Android_ASTCTargetPlatform.Build.cs b/Engine/Source/Developer/Android/Android_ASTCTargetPlatform/Android_ASTCTargetPlatform.Build.cs
new file mode 100644
--- /dev/null
+++ b/Engine/Source/Developer/Android/Android_ASTCTargetPlatform/Android_ASTCTargetPlatform.Build.cs
@@ -0,0 +1,43 @@
+// Copyright 1998-2014 Epic Games, Inc. All Rights Reserved.
+
+using UnrealBuildTool;
+
+public class Android_ASTCTargetPlatform : ModuleRules
+{
+	public Android_ASTCTargetPlatform(TargetInfo Target)
+	{
+		BinariesSubFolder = "Android";
+
+        PrivateDependencyModuleNames.AddRange(
+			new string[]
+			{
+				"Core",
+				"TargetPlatform",
+				"DesktopPlatform",
+				"AndroidDeviceDetection",
+			}
+		);
+
+		PublicIncludePaths.AddRange(
+			new string[]
+			{
+				"Runtime/Core/Public/Android"
+			}
+		);
+
+		if (UEBuildConfiguration.bCompileAgainstEngine)
+		{
+			PrivateDependencyModuleNames.Add("Engine");
+			PrivateIncludePathModuleNames.Add("TextureCompressor");		//@todo android: AndroidTargetPlatform.Build
+		}
+
+        Definitions.Add("WITH_OGGVORBIS=1");
+
+		PrivateIncludePaths.AddRange(
+			new string[]
+			{
+				"Developer/Android/AndroidTargetPlatform/Private",				
+			}
+		);
+	}
+}
diff --git a/Engine/Source/Developer/Android/Android_ASTCTargetPlatform/Private/Android_ASTCTargetPlatformModule.cpp b/Engine/Source/Developer/Android/Android_ASTCTargetPlatform/Private/Android_ASTCTargetPlatformModule.cpp
new file mode 100644
--- /dev/null
+++ b/Engine/Source/Developer/Android/Android_ASTCTargetPlatform/Private/Android_ASTCTargetPlatformModule.cpp
@@ -0,0 +1,157 @@
+// Copyright 1998-2014 Epic Games, Inc. All Rights Reserved.
+
+/*=============================================================================
+	AndroidTargetPlatformModule.cpp: Implements the FAndroidTargetPlatformModule class.
+=============================================================================*/
+
+#include "Android_ASTCTargetPlatformPrivatePCH.h"
+
+#define LOCTEXT_NAMESPACE "FAndroid_ASTCTargetPlatformModule" 
+
+/**
+ * Android cooking platform which cooks only ASTC based textures.
+ */
+class FAndroid_ASTCTargetPlatform
+	: public FAndroidTargetPlatform<FAndroid_ASTCPlatformProperties>
+{
+	virtual FString GetAndroidVariantName( ) override
+	{
+		return TEXT("ASTC");
+	}
+
+	virtual FText DisplayName( ) const override
+	{
+		return LOCTEXT("Android_ASTC", "Android (ASTC)");
+	}
+
+	virtual FString PlatformName() const override
+	{
+		return FString(FAndroid_ASTCPlatformProperties::PlatformName());
+	}
+
+	virtual bool SupportsTextureFormat( FName Format ) const override
+	{
+		if( Format == AndroidTexFormat::NameASTC_4x4 ||
+			Format == AndroidTexFormat::NameASTC_6x6 ||
+			Format == AndroidTexFormat::NameASTC_8x8 ||
+			Format == AndroidTexFormat::NameASTC_10x10 ||
+			Format == AndroidTexFormat::NameASTC_12x12 ||
+			Format == AndroidTexFormat::NameAutoASTC )
+		{
+			return true;
+		}
+		return false;
+	}
+
+	virtual void GetTextureFormats(const UTexture* Texture, TArray<FName>& OutFormats) const
+	{
+		check(Texture);
+
+		// we remap some of the defaults (with PVRTC and ASTC formats)
+		static FName FormatRemap[][2] =
+		{
+			// Default format:				ASTC format:
+			{ { FName(TEXT("DXT1")) },		{ FName(TEXT("ASTC_RGB")) } },
+			{ { FName(TEXT("DXT5")) },		{ FName(TEXT("ASTC_RGBA")) } },
+			{ { FName(TEXT("DXT5n")) },		{ FName(TEXT("ASTC_NormalAG")) } },
+			{ { FName(TEXT("BC5")) },		{ FName(TEXT("ASTC_NormalRG")) } },
+			{ { FName(TEXT("BC6H")) },		{ FName(TEXT("ASTC_RGB")) } },
+			{ { FName(TEXT("BC7")) },		{ FName(TEXT("ASTC_RGBAuto")) } },
+			{ { FName(TEXT("AutoDXT")) },	{ FName(TEXT("ASTC_RGBAuto")) } },
+		};
+
+		FName TextureFormatName = NAME_None;
+
+		// forward rendering only needs one channel for shadow maps
+		if (Texture->LODGroup == TEXTUREGROUP_Shadowmap)
+		{
+			TextureFormatName = FName(TEXT("G8"));
+		}
+
+		// if we didn't assign anything specially, then use the defaults
+		if (TextureFormatName == NAME_None)
+		{
+			TextureFormatName = GetDefaultTextureFormatName(Texture, EngineSettings, false);
+		}
+
+		// perform any remapping away from defaults
+		bool bFoundRemap = false;
+		for (int32 RemapIndex = 0; RemapIndex < ARRAY_COUNT(FormatRemap); ++RemapIndex)
+		{
+			if (TextureFormatName == FormatRemap[RemapIndex][0])
+			{
+				// we found a remapping
+				bFoundRemap = true;
+				OutFormats.AddUnique(FormatRemap[RemapIndex][1]);
+			}
+		}
+
+		// if we didn't already remap above, add it now
+		if (!bFoundRemap)
+		{
+			OutFormats.Add(TextureFormatName);
+		}
+	}
+
+
+	virtual bool SupportedByExtensionsString( const FString& ExtensionsString, const int GLESVersion ) const override
+	{
+		return ExtensionsString.Contains(TEXT("GL_KHR_texture_compression_astc_ldr"));
+	}
+
+	virtual FText GetVariantDisplayName() const override
+	{
+		return LOCTEXT("Android_ASTC_ShortName", "ASTC");
+	}
+
+	virtual float GetVariantPriority() const override
+	{
+		return 0.9f;
+	}
+};
+
+/**
+ * Holds the target platform singleton.
+ */
+static ITargetPlatform* AndroidTargetSingleton = NULL;
+
+
+/**
+ * Module for the Android target platform.
+ */
+class FAndroid_ASTCTargetPlatformModule
+	: public ITargetPlatformModule
+{
+public:
+
+	/**
+	 * Destructor.
+	 */
+	~FAndroid_ASTCTargetPlatformModule()
+	{
+		AndroidTargetSingleton = NULL;
+	}
+
+
+public:
+	
+	// Begin ITargetPlatformModule interface
+
+	virtual ITargetPlatform* GetTargetPlatform() override
+	{
+		if (AndroidTargetSingleton == NULL)
+		{
+			AndroidTargetSingleton = new FAndroid_ASTCTargetPlatform();
+		}
+		
+		return AndroidTargetSingleton;
+	}
+
+	// End ITargetPlatformModule interface
+};
+
+
+#undef LOCTEXT_NAMESPACE
+
+
+IMPLEMENT_MODULE(FAndroid_ASTCTargetPlatformModule, Android_ASTCTargetPlatform);
diff --git a/Engine/Source/Developer/Android/Android_ASTCTargetPlatform/Private/Android_ASTCTargetPlatformPrivatePCH.h b/Engine/Source/Developer/Android/Android_ASTCTargetPlatform/Private/Android_ASTCTargetPlatformPrivatePCH.h
new file mode 100644
--- /dev/null
+++ b/Engine/Source/Developer/Android/Android_ASTCTargetPlatform/Private/Android_ASTCTargetPlatformPrivatePCH.h
@@ -0,0 +1,32 @@
+// Copyright 1998-2014 Epic Games, Inc. All Rights Reserved.
+
+
+/*=============================================================================
+	Android_ASTCTargetPlatformPrivatePCH.h: Pre-compiled header file for the Android_ASTCTargetPlatform module.
+=============================================================================*/
+
+#pragma once
+
+/* Dependencies
+ *****************************************************************************/
+
+#include "Core.h"
+#include "ModuleInterface.h"
+#include "ModuleManager.h"
+#include "Runtime/Core/Public/Android/AndroidProperties.h"
+
+#if WITH_ENGINE
+	#include "Engine.h"
+	#include "TextureCompressorModule.h"
+#endif
+
+#include "TargetPlatform.h"
+#include "TargetPlatformBase.h"
+#include "AndroidDeviceDetection.h"
+
+
+/* Private includes
+ *****************************************************************************/
+
+#include "AndroidTargetDevice.h"
+#include "AndroidTargetPlatform.h"
diff --git a/Engine/Source/Developer/DesktopPlatform/Private/PlatformInfo.cpp b/Engine/Source/Developer/DesktopPlatform/Private/PlatformInfo.cpp
--- a/Engine/Source/Developer/DesktopPlatform/Private/PlatformInfo.cpp
+++ b/Engine/Source/Developer/DesktopPlatform/Private/PlatformInfo.cpp
@@ -102,6 +102,7 @@ static const FPlatformInfo PlatformInfoArray[] = {
 	BuildPlatformInfo(TEXT("Android_ETC1"),				TEXT("Android_ETC1"),		LOCTEXT("Android_ETC1", "Android (ETC1)"),					EPlatformType::Game,		EPlatformFlags::CookFlavor,		FPlatformIconPaths(TEXT("Launcher/AndroidTarget/Platform_Android_ETC1_24x"), TEXT("Launcher/AndroidTarget/Platform_Android_128x")),					TEXT("-targetplatform=Android -cookflavor=ETC1"),		TEXT("Android"),	EPlatformSDKStatus::Unknown,	TEXT("/Engine/Tutorial/Mobile/SettingUpAndroidTutorial.SettingUpAndroidTutorial"),						IsAvailableOnWindows || IsAvailableOnMac,	true,				false),
 	BuildPlatformInfo(TEXT("Android_ETC2"),				TEXT("Android_ETC2"),		LOCTEXT("Android_ETC2", "Android (ETC2)"),					EPlatformType::Game,		EPlatformFlags::CookFlavor,		FPlatformIconPaths(TEXT("Launcher/AndroidTarget/Platform_Android_ETC2_24x"), TEXT("Launcher/AndroidTarget/Platform_Android_128x")),					TEXT("-targetplatform=Android -cookflavor=ETC2"),		TEXT("Android"),	EPlatformSDKStatus::Unknown,	TEXT("/Engine/Tutorial/Mobile/SettingUpAndroidTutorial.SettingUpAndroidTutorial"),						IsAvailableOnWindows || IsAvailableOnMac,	true,				false),
 	BuildPlatformInfo(TEXT("Android_PVRTC"),			TEXT("Android_PVRTC"),		LOCTEXT("Android_PVRTC", "Android (PVRTC)"),				EPlatformType::Game,		EPlatformFlags::CookFlavor,		FPlatformIconPaths(TEXT("Launcher/AndroidTarget/Platform_Android_PVRTC_24x"), TEXT("Launcher/AndroidTarget/Platform_Android_128x")),				TEXT("-targetplatform=Android -cookflavor=PVRTC"),		TEXT("Android"),	EPlatformSDKStatus::Unknown,	TEXT("/Engine/Tutorial/Mobile/SettingUpAndroidTutorial.SettingUpAndroidTutorial"),						IsAvailableOnWindows || IsAvailableOnMac,	true,				false),
+	BuildPlatformInfo(TEXT("Android_ASTC"),				TEXT("Android_ASTC"),		LOCTEXT("Android_ASTC", "Android (ASTC)"),					EPlatformType::Game,		EPlatformFlags::CookFlavor,		FPlatformIconPaths(TEXT("Launcher/AndroidTarget/Platform_Android_ASTC_24x"), TEXT("Launcher/AndroidTarget/Platform_Android_128x")),					TEXT("-targetplatform=Android -cookflavor=ASTC"),		TEXT("Android"),	EPlatformSDKStatus::Unknown,	TEXT("/Engine/Tutorial/Mobile/SettingUpAndroidTutorial.SettingUpAndroidTutorial"),						IsAvailableOnWindows || IsAvailableOnMac,	true,				false),
 	// @todo es31: Set the final false to true when ES31 is in Binary releases
 	BuildPlatformInfo(TEXT("Android_ES31"),				TEXT("Android_ES31"),		LOCTEXT("Android_ES31", "Android (ES31+AEP)"),				EPlatformType::Game,		EPlatformFlags::CookFlavor,		FPlatformIconPaths(TEXT("Launcher/AndroidTarget/Platform_Android_ES31_24x"), TEXT("Launcher/AndroidTarget/Platform_Android_128x")),					TEXT("-targetplatform=Android -cookflavor=ES31"),		TEXT("Android"),	EPlatformSDKStatus::Unknown,	TEXT("Shared/Tutorials/SettingUpAndroidTutorial"),														IsAvailableOnWindows,						false,				false),
 //	BuildPlatformInfo(TEXT("Android_GL4"),				TEXT("Android_GL4"),		LOCTEXT("Android_GL4", "Android (GL4)"),					EPlatformType::Game,		EPlatformFlags::CookFlavor,		FPlatformIconPaths(TEXT("Launcher/AndroidTarget/Platform_Android_GL4_24x"), TEXT("Launcher/AndroidTarget/Platform_Android_128x")),					TEXT("-targetplatform=Android -cookflavor=GL4"),		TEXT("Android"),	EPlatformSDKStatus::Unknown,	TEXT("Shared/Tutorials/SettingUpAndroidTutorial"),														IsAvailableOnWindows,						true),
diff --git a/Engine/Source/Developer/TargetPlatform/Public/Common/TargetPlatformBase.h b/Engine/Source/Developer/TargetPlatform/Public/Common/TargetPlatformBase.h
--- a/Engine/Source/Developer/TargetPlatform/Public/Common/TargetPlatformBase.h
+++ b/Engine/Source/Developer/TargetPlatform/Public/Common/TargetPlatformBase.h
@@ -175,7 +175,7 @@ public:
 	}
 #else //TEXTURERESOURCE_H_INCLUDED
 
-	FName GetDefaultTextureFormatName( const UTexture* Texture, const FConfigFile& EngineSettings ) const
+	FName GetDefaultTextureFormatName( const UTexture* Texture, const FConfigFile& EngineSettings, bool bSupportDX11TextureFormats ) const
 	{
 		return NAME_None;
 	}
diff --git a/Engine/Source/Developer/TargetPlatform/TargetPlatform.Build.cs b/Engine/Source/Developer/TargetPlatform/TargetPlatform.Build.cs
--- a/Engine/Source/Developer/TargetPlatform/TargetPlatform.Build.cs
+++ b/Engine/Source/Developer/TargetPlatform/TargetPlatform.Build.cs
@@ -60,6 +60,7 @@ public class TargetPlatform : ModuleRules
 					DynamicallyLoadedModuleNames.Add("Android_DXTTargetPlatform");
 					DynamicallyLoadedModuleNames.Add("Android_ETC1TargetPlatform");
 					DynamicallyLoadedModuleNames.Add("Android_ETC2TargetPlatform");
+					DynamicallyLoadedModuleNames.Add("Android_ASTCTargetPlatform");
 					DynamicallyLoadedModuleNames.Add("IOSTargetPlatform");
                     DynamicallyLoadedModuleNames.Add("HTML5TargetPlatform");
 				}
diff --git a/Engine/Source/Developer/TextureFormatASTC/Private/TextureFormatASTC.cpp b/Engine/Source/Developer/TextureFormatASTC/Private/TextureFormatASTC.cpp
--- a/Engine/Source/Developer/TextureFormatASTC/Private/TextureFormatASTC.cpp
+++ b/Engine/Source/Developer/TextureFormatASTC/Private/TextureFormatASTC.cpp
@@ -13,7 +13,7 @@
 
 
 // increment this if you change anything that will affect compression in this file, including FORCED_NORMAL_MAP_COMPRESSION_SIZE_VALUE
-#define BASE_FORMAT_VERSION 30
+#define BASE_FORMAT_VERSION 34
 
 #define MAX_QUALITY_BY_SIZE 4
 #define MAX_QUALITY_BY_SPEED 3
@@ -151,15 +151,16 @@ static bool CompressSliceToASTC(
 )
 {
 	// Always Y-invert the image prior to compression for proper orientation post-compression
-	uint8 LineBuffer[16384 * 4];
+	TArray<uint8> LineBuffer;
+	LineBuffer.AddUninitialized(16384 * 4);
 	uint32 LineSize = SizeX * 4;
 	for (int32 LineIndex = 0; LineIndex < (SizeY / 2); LineIndex++)
 	{
 		uint8* LineData0 = ((uint8*)SourceData) + (LineSize * LineIndex);
 		uint8* LineData1 = ((uint8*)SourceData) + (LineSize * (SizeY - LineIndex - 1));
-		FMemory::Memcpy(LineBuffer, LineData0,  LineSize);
-		FMemory::Memcpy(LineData0,  LineData1,  LineSize);
-		FMemory::Memcpy(LineData1,  LineBuffer, LineSize);
+		FMemory::Memcpy(LineBuffer.GetData(), LineData0, LineSize);
+		FMemory::Memcpy(LineData0, LineData1, LineSize);
+		FMemory::Memcpy(LineData1, LineBuffer.GetData(), LineSize);
 	}
 	
 	// Compress and retrieve the PNG data to write out to disk
@@ -186,7 +187,7 @@ static bool CompressSliceToASTC(
 	delete PNGFile;
 
 	// Compress PNG file to ASTC (using the reference astcenc.exe from ARM)
-	FString Params = FString::Printf(TEXT("-c %s %s %s"),
+	FString Params = FString::Printf(TEXT("-c \"%s\" \"%s\" %s"),
 		*InputFilePath,
 		*OutputFilePath,
 		*CompressionParameters
@@ -344,13 +345,13 @@ class FTextureFormatASTC : public ITextureFormat
 		  ((BuildSettings.TextureFormatName == GTextureFormatNameASTC_RGBAuto) && !bImageHasAlphaChannel))
 		{
 			CompressedPixelFormat = GetQualityFormat();
-			CompressionParameters = FString::Printf(TEXT("%s -esw bgra -ch 1 1 1 0"), *GetQualityString());
+			CompressionParameters = FString::Printf(TEXT("%s %s -esw bgra -ch 1 1 1 0"), *GetQualityString(), /*BuildSettings.bSRGB ? TEXT("-srgb") :*/ TEXT("") );
 		}
 		else if (BuildSettings.TextureFormatName == GTextureFormatNameASTC_RGBA ||
 			   ((BuildSettings.TextureFormatName == GTextureFormatNameASTC_RGBAuto) && bImageHasAlphaChannel))
 		{
 			CompressedPixelFormat = GetQualityFormat();
-			CompressionParameters = FString::Printf(TEXT("%s -esw bgra -ch 1 1 1 1 -alphablend"), *GetQualityString());
+			CompressionParameters = FString::Printf(TEXT("%s %s -esw bgra -ch 1 1 1 1 -alphablend"), *GetQualityString(), /*BuildSettings.bSRGB ? TEXT("-srgb") :*/ TEXT("") );
 		}
 		else if (BuildSettings.TextureFormatName == GTextureFormatNameASTC_NormalAG)
 		{
diff --git a/Engine/Source/Programs/UnrealBuildTool/Android/UEBuildAndroid.cs b/Engine/Source/Programs/UnrealBuildTool/Android/UEBuildAndroid.cs
--- a/Engine/Source/Programs/UnrealBuildTool/Android/UEBuildAndroid.cs
+++ b/Engine/Source/Programs/UnrealBuildTool/Android/UEBuildAndroid.cs
@@ -303,6 +303,7 @@ namespace UnrealBuildTool
 							InModule.AddPlatformSpecificDynamicallyLoadedModule("Android_ETC1TargetPlatform");
 							InModule.AddPlatformSpecificDynamicallyLoadedModule("Android_ETC2TargetPlatform");
 							InModule.AddPlatformSpecificDynamicallyLoadedModule("Android_ES31TargetPlatform");
+							InModule.AddPlatformSpecificDynamicallyLoadedModule("Android_ASTCTargetPlatform");
 							// @todo gl4android				
 							// InModule.AddPlatformSpecificDynamicallyLoadedModule("Android_GL4TargetPlatform");
                         }
@@ -331,6 +332,7 @@ namespace UnrealBuildTool
 					InModule.AddPlatformSpecificDynamicallyLoadedModule("Android_ETC1TargetPlatform");
 					InModule.AddPlatformSpecificDynamicallyLoadedModule("Android_ETC2TargetPlatform");
 					InModule.AddPlatformSpecificDynamicallyLoadedModule("Android_ES31TargetPlatform");
+					InModule.AddPlatformSpecificDynamicallyLoadedModule("Android_ASTCTargetPlatform");
 					// @todo gl4android				
 					// InModule.AddPlatformSpecificDynamicallyLoadedModule("Android_GL4TargetPlatform");
                 }
diff --git a/Engine/Source/Runtime/Core/Public/Android/AndroidProperties.h b/Engine/Source/Runtime/Core/Public/Android/AndroidProperties.h
--- a/Engine/Source/Runtime/Core/Public/Android/AndroidProperties.h
+++ b/Engine/Source/Runtime/Core/Public/Android/AndroidProperties.h
@@ -123,6 +123,12 @@ struct FAndroidPlatformProperties
 	{
 		return true;
 	}
+
+	static FORCEINLINE const FString& GetCookFlavor()
+	{
+		extern FString GCookFlavor;
+		return GCookFlavor;
+	}
 };
 
 struct FAndroid_PVRTCPlatformProperties : public FAndroidPlatformProperties
@@ -165,6 +171,14 @@ struct FAndroid_ETC2PlatformProperties : public FAndroidPlatformProperties
 	}
 };
 
+struct FAndroid_ASTCPlatformProperties : public FAndroidPlatformProperties
+{
+	static FORCEINLINE const char* PlatformName()
+	{
+		return "Android_ASTC";
+	}
+};
+
 struct FAndroid_GL4PlatformProperties : public FAndroidPlatformProperties
 {
 	static FORCEINLINE const char* PlatformName()
diff --git a/Engine/Source/Runtime/Engine/Public/PixelFormat.h b/Engine/Source/Runtime/Engine/Public/PixelFormat.h
--- a/Engine/Source/Runtime/Engine/Public/PixelFormat.h
+++ b/Engine/Source/Runtime/Engine/Public/PixelFormat.h
@@ -44,23 +44,23 @@ enum EPixelFormat
 	PF_R16G16B16A16_SINT    =35,
 	PF_R5G6B5_UNORM         =36,
 	PF_R8G8B8A8             =37,
-	PF_A8R8G8B8				=38,			// Only used for legacy loading; do NOT use!
+	PF_A8R8G8B8				=38,	// Only used for legacy loading; do NOT use!
 	PF_BC4					=39,
 	PF_R8G8                 =40,
 	PF_ATC_RGB				=41,
 	PF_ATC_RGBA_E			=42,
 	PF_ATC_RGBA_I			=43,
-	PF_X24_G8				=44,			// Used for creating SRVs to alias a DepthStencil buffer to read Stencil. Don't use for creating textures.
+	PF_X24_G8				=44,	// Used for creating SRVs to alias a DepthStencil buffer to read Stencil. Don't use for creating textures.
 	PF_ETC1					=45,
 	PF_ETC2_RGB				=46,
 	PF_ETC2_RGBA			=47,
 	PF_R32G32B32A32_UINT	=48,
 	PF_R16G16_UINT			=49,
-	PF_ASTC_4x4             =50,
-	PF_ASTC_6x6             =51,
-	PF_ASTC_8x8             =52,
-	PF_ASTC_10x10           =53,
-	PF_ASTC_12x12           =54,
+	PF_ASTC_4x4             =50,	// 8.00 bpp
+	PF_ASTC_6x6             =51,	// 3.56 bpp
+	PF_ASTC_8x8             =52,	// 2.00 bpp
+	PF_ASTC_10x10           =53,	// 1.28 bpp
+	PF_ASTC_12x12           =54,	// 0.89 bpp
 	PF_BC6H					=55,
 	PF_BC7					=56,
 	PF_MAX                  =57,
diff --git a/Engine/Source/Runtime/OpenGLDrv/Private/Android/AndroidOpenGLPrivate.h b/Engine/Source/Runtime/OpenGLDrv/Private/Android/AndroidOpenGLPrivate.h
--- a/Engine/Source/Runtime/OpenGLDrv/Private/Android/AndroidOpenGLPrivate.h
+++ b/Engine/Source/Runtime/OpenGLDrv/Private/Android/AndroidOpenGLPrivate.h
@@ -86,6 +86,10 @@ private:
 		{
 			TargetPlatformNames.Add(TEXT("Android_ETC2"));
 		}
+		if (ExtensionsString.Contains(TEXT("GL_KHR_texture_compression_astc_ldr")))
+		{
+			TargetPlatformNames.Add(TEXT("Android_ASTC"));
+		}
 
 		// all devices support ETC
 		TargetPlatformNames.Add(TEXT("Android_ETC1"));
diff --git a/Engine/Source/Runtime/OpenGLDrv/Private/OpenGLDevice.cpp b/Engine/Source/Runtime/OpenGLDrv/Private/OpenGLDevice.cpp
--- a/Engine/Source/Runtime/OpenGLDrv/Private/OpenGLDevice.cpp
+++ b/Engine/Source/Runtime/OpenGLDrv/Private/OpenGLDevice.cpp
@@ -877,6 +877,14 @@ static void InitRHICapabilitiesForGL()
 		SetupTextureFormat( PF_ETC2_RGBA,	FOpenGLTextureFormat(GL_COMPRESSED_RGBA8_ETC2_EAC,	FOpenGL::SupportsSRGB() ? GL_COMPRESSED_SRGB8_ALPHA8_ETC2_EAC : GL_COMPRESSED_RGBA8_ETC2_EAC,	GL_RGBA,	GL_UNSIGNED_BYTE,	true,			false));
 	}
 #endif
+	if (FOpenGL::SupportsASTC())
+	{
+		SetupTextureFormat( PF_ASTC_4x4,	FOpenGLTextureFormat(GL_COMPRESSED_RGBA_ASTC_4x4_KHR,	GL_COMPRESSED_SRGB8_ALPHA8_ASTC_4x4_KHR,	GL_RGBA,	GL_UNSIGNED_BYTE,	true,	false) );
+		SetupTextureFormat( PF_ASTC_6x6,	FOpenGLTextureFormat(GL_COMPRESSED_RGBA_ASTC_6x6_KHR,	GL_COMPRESSED_SRGB8_ALPHA8_ASTC_6x6_KHR,	GL_RGBA,	GL_UNSIGNED_BYTE,	true,	false) );
+		SetupTextureFormat( PF_ASTC_8x8,	FOpenGLTextureFormat(GL_COMPRESSED_RGBA_ASTC_8x8_KHR,	GL_COMPRESSED_SRGB8_ALPHA8_ASTC_8x8_KHR,	GL_RGBA,	GL_UNSIGNED_BYTE,	true,	false) );
+		SetupTextureFormat( PF_ASTC_10x10,	FOpenGLTextureFormat(GL_COMPRESSED_RGBA_ASTC_10x10_KHR,	GL_COMPRESSED_SRGB8_ALPHA8_ASTC_10x10_KHR,	GL_RGBA,	GL_UNSIGNED_BYTE,	true,	false) );
+		SetupTextureFormat( PF_ASTC_12x12,	FOpenGLTextureFormat(GL_COMPRESSED_RGBA_ASTC_12x12_KHR,	GL_COMPRESSED_SRGB8_ALPHA8_ASTC_12x12_KHR,	GL_RGBA,	GL_UNSIGNED_BYTE,	true,	false) );
+	}
 
 	// Some formats need to know how large a block is.
 	GPixelFormats[ PF_DepthStencil		].BlockBytes	 = 4;
