- Don't tie some Android editor project settings to having the Android support files

diff --git a/Engine/Source/Developer/Android/AndroidPlatformEditor/Private/AndroidTargetSettingsCustomization.cpp b/Engine/Source/Developer/Android/AndroidPlatformEditor/Private/AndroidTargetSettingsCustomization.cpp
--- a/Engine/Source/Developer/Android/AndroidPlatformEditor/Private/AndroidTargetSettingsCustomization.cpp
+++ b/Engine/Source/Developer/Android/AndroidPlatformEditor/Private/AndroidTargetSettingsCustomization.cpp
@@ -196,7 +196,6 @@ void FAndroidTargetSettingsCustomization::BuildAppManifestSection(IDetailLayoutB
 	{ \
 		TSharedRef<IPropertyHandle> PropertyHandle = DetailLayout.GetProperty(GET_MEMBER_NAME_CHECKED(UAndroidRuntimeSettings, PropName)); \
 		Category.AddProperty(PropertyHandle) \
-			.EditCondition(SetupForPlatformAttribute, NULL) \
 			.IsEnabled(!FRocketSupport::IsRocket()) \
 			.ToolTip(!FRocketSupport::IsRocket() ? Tip : FAndroidTargetSettingsCustomizationConstants::DisabledTip); \
 	}
