CIS fix for UFE

diff --git a/Engine/Source/Developer/Android/Android_ASTCTargetPlatform/Private/Android_ASTCTargetPlatformModule.cpp b/Engine/Source/Developer/Android/Android_ASTCTargetPlatform/Private/Android_ASTCTargetPlatformModule.cpp
--- a/Engine/Source/Developer/Android/Android_ASTCTargetPlatform/Private/Android_ASTCTargetPlatformModule.cpp
+++ b/Engine/Source/Developer/Android/Android_ASTCTargetPlatform/Private/Android_ASTCTargetPlatformModule.cpp
@@ -62,6 +62,7 @@ class FAndroid_ASTCTargetPlatform
 
 		FName TextureFormatName = NAME_None;
 
+#if WITH_EDITOR
 		// forward rendering only needs one channel for shadow maps
 		if (Texture->LODGroup == TEXTUREGROUP_Shadowmap)
 		{
@@ -91,6 +92,9 @@ class FAndroid_ASTCTargetPlatform
 		{
 			OutFormats.Add(TextureFormatName);
 		}
+#else
+		OutFormats.Add(TextureFormatName);
+#endif
 	}
 
 
