Better fix for UnrealFrontEnd compile error

diff --git a/Engine/Source/Developer/Android/Android_ASTCTargetPlatform/Private/Android_ASTCTargetPlatformModule.cpp b/Engine/Source/Developer/Android/Android_ASTCTargetPlatform/Private/Android_ASTCTargetPlatformModule.cpp
--- a/Engine/Source/Developer/Android/Android_ASTCTargetPlatform/Private/Android_ASTCTargetPlatformModule.cpp
+++ b/Engine/Source/Developer/Android/Android_ASTCTargetPlatform/Private/Android_ASTCTargetPlatformModule.cpp
@@ -43,6 +43,7 @@ class FAndroid_ASTCTargetPlatform
 		return false;
 	}
 
+#if WITH_ENGINE
 	virtual void GetTextureFormats(const UTexture* Texture, TArray<FName>& OutFormats) const
 	{
 		check(Texture);
@@ -62,7 +63,6 @@ class FAndroid_ASTCTargetPlatform
 
 		FName TextureFormatName = NAME_None;
 
-#if WITH_EDITOR
 		// forward rendering only needs one channel for shadow maps
 		if (Texture->LODGroup == TEXTUREGROUP_Shadowmap)
 		{
@@ -92,11 +92,8 @@ class FAndroid_ASTCTargetPlatform
 		{
 			OutFormats.Add(TextureFormatName);
 		}
-#else
-		OutFormats.Add(TextureFormatName);
-#endif
 	}
-
+#endif
 
 	virtual bool SupportedByExtensionsString( const FString& ExtensionsString, const int GLESVersion ) const override
 	{
