- Changed Android ES31+AEP support from a TargetPlatform to a checkbox that works with any Android TP
- Moved HighQualityLightmap support check to IOSTargetPlatform so it can check project settings for Metal support
- Fixed some issues running on an x86_64 device
- Cleaned up some unused functions in TargetPlatform

diff --git a/Engine/Build/JunkManifest.txt b/Engine/Build/JunkManifest.txt
--- a/Engine/Build/JunkManifest.txt
+++ b/Engine/Build/JunkManifest.txt
@@ -53,3 +53,6 @@ Engine\Plugins\Script\ScriptPlugin\Binaries\Win32\*
 Engine\Plugins\Script\ScriptPlugin\Binaries\Win64\*
 Engine\Plugins\Experimental\Paper2D\Binaries
 Engine\Plugins\Experimental\Paper2D\Intermediate
+Engine\Binaries\Mac\UE4Editor.app\Contents\MacOS\UE4Editor-Android_ES31TargetPlatform.dylib
+Engine\Binaries\Mac\UnrealFrontend.app\Contents\MacOS\UnrealFrontend-Android_ES31TargetPlatform.dylib
+Engine\Binaries\Win64\Android\*Android_ES31*.*
diff --git a/Engine/Source/Developer/Android/AndroidTargetPlatform/Private/AndroidTargetPlatform.h b/Engine/Source/Developer/Android/AndroidTargetPlatform/Private/AndroidTargetPlatform.h
--- a/Engine/Source/Developer/Android/AndroidTargetPlatform/Private/AndroidTargetPlatform.h
+++ b/Engine/Source/Developer/Android/AndroidTargetPlatform/Private/AndroidTargetPlatform.h
@@ -121,10 +121,7 @@ public:
 	}
 
 #if WITH_ENGINE
-	virtual void GetReflectionCaptureFormats( TArray<FName>& OutFormats ) const override
-	{
-		OutFormats.Add(FName(TEXT("EncodedHDR")));
-	}
+	virtual void GetReflectionCaptureFormats( TArray<FName>& OutFormats ) const override;
 
 	virtual void GetAllPossibleShaderFormats( TArray<FName>& OutFormats ) const override;
 
diff --git a/Engine/Source/Developer/Android/AndroidTargetPlatform/Private/AndroidTargetPlatform.inl b/Engine/Source/Developer/Android/AndroidTargetPlatform/Private/AndroidTargetPlatform.inl
--- a/Engine/Source/Developer/Android/AndroidTargetPlatform/Private/AndroidTargetPlatform.inl
+++ b/Engine/Source/Developer/Android/AndroidTargetPlatform/Private/AndroidTargetPlatform.inl
@@ -10,6 +10,22 @@
 
 #define LOCTEXT_NAMESPACE "FAndroidTargetPlatform"
 
+static bool SupportsES2()
+{
+	// default to support ES2
+	bool bBuildForES2 = true;
+	GConfig->GetBool(TEXT("/Script/AndroidRuntimeSettings.AndroidRuntimeSettings"), TEXT("bBuildForES2"), bBuildForES2, GEngineIni);
+	return bBuildForES2;
+}
+
+static bool SupportsAEP()
+{
+	// default to not supporting ES31
+	bool bBuildForES31 = false;
+	GConfig->GetBool(TEXT("/Script/AndroidRuntimeSettings.AndroidRuntimeSettings"), TEXT("bBuildForES31"), bBuildForES31, GEngineIni);
+	return bBuildForES31;
+}
+
 template<class TPlatformProperties>
 inline FAndroidTargetPlatform<TPlatformProperties>::FAndroidTargetPlatform( ) :
 	DeviceDetection(nullptr)
@@ -89,7 +105,6 @@ inline bool FAndroidTargetPlatform<TPlatformProperties>::IsRunningPlatform( ) co
 template<class TPlatformProperties>
 inline bool FAndroidTargetPlatform<TPlatformProperties>::IsSdkInstalled(bool bProjectHasCode, FString& OutDocumentationPath) const
 {
-
 	OutDocumentationPath = FString("Shared/Tutorials/SettingUpAndroidTutorial");
 	return true;
 }
@@ -98,11 +113,23 @@ inline bool FAndroidTargetPlatform<TPlatformProperties>::IsSdkInstalled(bool bPr
 template<class TPlatformProperties>
 inline bool FAndroidTargetPlatform<TPlatformProperties>::SupportsFeature( ETargetPlatformFeatures Feature ) const
 {
-	if (Feature == ETargetPlatformFeatures::Packaging)
+	switch (Feature)
 	{
-		return true;
+		case ETargetPlatformFeatures::Packaging:
+			return true;
+			
+		case ETargetPlatformFeatures::LowQualityLightmaps:
+			return SupportsES2();
+			
+		case ETargetPlatformFeatures::HighQualityLightmaps:
+		case ETargetPlatformFeatures::VertexShaderTextureSampling:
+		case ETargetPlatformFeatures::Tessellation:
+			return SupportsAEP();
+			
+		default:
+			break;
 	}
-
+	
 	return TTargetPlatformBase<TPlatformProperties>::SupportsFeature(Feature);
 }
 
@@ -113,7 +140,22 @@ template<class TPlatformProperties>
 inline void FAndroidTargetPlatform<TPlatformProperties>::GetAllPossibleShaderFormats( TArray<FName>& OutFormats ) const
 {
 	static FName NAME_OPENGL_ES2(TEXT("GLSL_ES2"));
-	OutFormats.AddUnique(NAME_OPENGL_ES2);
+	static FName NAME_GLSL_310_ES_EXT(TEXT("GLSL_310_ES_EXT"));
+
+	// get project rendering support, default to ES2 and no ES31
+	bool bBuildForES2 = true;
+	bool bBuildForES31 = false;
+	GConfig->GetBool(TEXT("/Script/AndroidRuntimeSettings.AndroidRuntimeSettings"), TEXT("bBuildForES2"), bBuildForES2, GEngineIni);
+	GConfig->GetBool(TEXT("/Script/AndroidRuntimeSettings.AndroidRuntimeSettings"), TEXT("bBuildForES31"), bBuildForES31, GEngineIni);
+
+	if (bBuildForES2)
+	{
+		OutFormats.AddUnique(NAME_OPENGL_ES2);
+	}
+	if (bBuildForES31)
+	{
+		OutFormats.AddUnique(NAME_GLSL_310_ES_EXT);
+	}
 }
 
 template<class TPlatformProperties>
@@ -226,6 +268,21 @@ inline void FAndroidTargetPlatform<TPlatformProperties>::GetTextureFormats( cons
 }
 
 
+template<class TPlatformProperties>
+void FAndroidTargetPlatform<TPlatformProperties>::GetReflectionCaptureFormats( TArray<FName>& OutFormats ) const
+{
+	if (SupportsAEP())
+	{
+		// use Full HDR with AEP
+		OutFormats.Add(FName(TEXT("FullHDR")));
+	}
+	
+	// always emit encoded
+	OutFormats.Add(FName(TEXT("EncodedHDR")));
+}
+
+
+
 template<class TPlatformProperties>
 const FTextureLODSettings& FAndroidTargetPlatform<TPlatformProperties>::GetTextureLODSettings( ) const
 {
diff --git a/Engine/Source/Developer/Android/Android_ES31TargetPlatform/Android_ES31TargetPlatform.Build.cs b/Engine/Source/Developer/Android/Android_ES31TargetPlatform/Android_ES31TargetPlatform.Build.cs
deleted file mode 100644
diff --git a/Engine/Source/Developer/Android/Android_ES31TargetPlatform/Private/Android_ES31TargetPlatformModule.cpp b/Engine/Source/Developer/Android/Android_ES31TargetPlatform/Private/Android_ES31TargetPlatformModule.cpp
deleted file mode 100644
diff --git a/Engine/Source/Developer/Android/Android_ES31TargetPlatform/Private/Android_ES31TargetPlatformPrivatePCH.h b/Engine/Source/Developer/Android/Android_ES31TargetPlatform/Private/Android_ES31TargetPlatformPrivatePCH.h
deleted file mode 100644
diff --git a/Engine/Source/Developer/DesktopPlatform/Private/PlatformInfo.cpp b/Engine/Source/Developer/DesktopPlatform/Private/PlatformInfo.cpp
--- a/Engine/Source/Developer/DesktopPlatform/Private/PlatformInfo.cpp
+++ b/Engine/Source/Developer/DesktopPlatform/Private/PlatformInfo.cpp
@@ -103,9 +103,6 @@ static const FPlatformInfo PlatformInfoArray[] = {
 	BuildPlatformInfo(TEXT("Android_ETC2"),				TEXT("Android_ETC2"),		LOCTEXT("Android_ETC2", "Android (ETC2)"),					EPlatformType::Game,		EPlatformFlags::CookFlavor,		FPlatformIconPaths(TEXT("Launcher/AndroidTarget/Platform_Android_ETC2_24x"), TEXT("Launcher/AndroidTarget/Platform_Android_128x")),					TEXT("-targetplatform=Android -cookflavor=ETC2"),		TEXT("Android"),	EPlatformSDKStatus::Unknown,	TEXT("/Engine/Tutorial/Mobile/SettingUpAndroidTutorial.SettingUpAndroidTutorial"),						IsAvailableOnWindows || IsAvailableOnMac,	true,				false),
 	BuildPlatformInfo(TEXT("Android_PVRTC"),			TEXT("Android_PVRTC"),		LOCTEXT("Android_PVRTC", "Android (PVRTC)"),				EPlatformType::Game,		EPlatformFlags::CookFlavor,		FPlatformIconPaths(TEXT("Launcher/AndroidTarget/Platform_Android_PVRTC_24x"), TEXT("Launcher/AndroidTarget/Platform_Android_128x")),				TEXT("-targetplatform=Android -cookflavor=PVRTC"),		TEXT("Android"),	EPlatformSDKStatus::Unknown,	TEXT("/Engine/Tutorial/Mobile/SettingUpAndroidTutorial.SettingUpAndroidTutorial"),						IsAvailableOnWindows || IsAvailableOnMac,	true,				false),
 	BuildPlatformInfo(TEXT("Android_ASTC"),				TEXT("Android_ASTC"),		LOCTEXT("Android_ASTC", "Android (ASTC)"),					EPlatformType::Game,		EPlatformFlags::CookFlavor,		FPlatformIconPaths(TEXT("Launcher/AndroidTarget/Platform_Android_ASTC_24x"), TEXT("Launcher/AndroidTarget/Platform_Android_128x")),					TEXT("-targetplatform=Android -cookflavor=ASTC"),		TEXT("Android"),	EPlatformSDKStatus::Unknown,	TEXT("/Engine/Tutorial/Mobile/SettingUpAndroidTutorial.SettingUpAndroidTutorial"),						IsAvailableOnWindows || IsAvailableOnMac,	true,				false),
-	// @todo es31: Set the final false to true when ES31 is in Binary releases
-	BuildPlatformInfo(TEXT("Android_ES31"),				TEXT("Android_ES31"),		LOCTEXT("Android_ES31", "Android (ES31+AEP)"),				EPlatformType::Game,		EPlatformFlags::CookFlavor,		FPlatformIconPaths(TEXT("Launcher/AndroidTarget/Platform_Android_ES31_24x"), TEXT("Launcher/AndroidTarget/Platform_Android_128x")),					TEXT("-targetplatform=Android -cookflavor=ES31"),		TEXT("Android"),	EPlatformSDKStatus::Unknown,	TEXT("Shared/Tutorials/SettingUpAndroidTutorial"),														IsAvailableOnWindows,						false,				false),
-//	BuildPlatformInfo(TEXT("Android_GL4"),				TEXT("Android_GL4"),		LOCTEXT("Android_GL4", "Android (GL4)"),					EPlatformType::Game,		EPlatformFlags::CookFlavor,		FPlatformIconPaths(TEXT("Launcher/AndroidTarget/Platform_Android_GL4_24x"), TEXT("Launcher/AndroidTarget/Platform_Android_128x")),					TEXT("-targetplatform=Android -cookflavor=GL4"),		TEXT("Android"),	EPlatformSDKStatus::Unknown,	TEXT("Shared/Tutorials/SettingUpAndroidTutorial"),														IsAvailableOnWindows,						true),
 
 	BuildPlatformInfo(TEXT("HTML5"),					TEXT("HTML5"),				LOCTEXT("HTML5", "HTML5"),									EPlatformType::Game,		EPlatformFlags::None,			FPlatformIconPaths(TEXT("Launcher/HTML5Target/Platform_HTML5_24x"), TEXT("Launcher/HTML5Target/Platform_HTML5_128x")),								TEXT(""),												TEXT("HTML5"),		EPlatformSDKStatus::Unknown,	TEXT("Platforms/HTML5/GettingStarted"),																	IsAvailableOnWindows || IsAvailableOnMac,   true,				false),
 
diff --git a/Engine/Source/Developer/TargetPlatform/Public/Common/TargetPlatformBase.h b/Engine/Source/Developer/TargetPlatform/Public/Common/TargetPlatformBase.h
--- a/Engine/Source/Developer/TargetPlatform/Public/Common/TargetPlatformBase.h
+++ b/Engine/Source/Developer/TargetPlatform/Public/Common/TargetPlatformBase.h
@@ -352,10 +352,6 @@ public:
 		return false;
 	}
 	
-	virtual uint32 MaxGpuSkinBones() const override
-	{
-		return TPlatformProperties::MaxGpuSkinBones();
-	}
 
 #if WITH_ENGINE
 	virtual FName GetPhysicsFormat( class UBodySetup* Body ) const override
diff --git a/Engine/Source/Developer/TargetPlatform/Public/Interfaces/ITargetPlatform.h b/Engine/Source/Developer/TargetPlatform/Public/Interfaces/ITargetPlatform.h
--- a/Engine/Source/Developer/TargetPlatform/Public/Interfaces/ITargetPlatform.h
+++ b/Engine/Source/Developer/TargetPlatform/Public/Interfaces/ITargetPlatform.h
@@ -229,13 +229,6 @@ public:
 	 */
 	virtual bool IsSdkInstalled(bool bProjectHasCode, FString& OutDocumentationPath) const = 0;
 
-	/**
-	 * Returns the maximum bones the platform supports.
-	 *
-	 * @return the maximum bones the platform supports.
-	 */
-	virtual uint32 MaxGpuSkinBones() const = 0;
-
 	/**
 	 * Checks whether this platform requires cooked data (typically console platforms).
 	 *
diff --git a/Engine/Source/Developer/iOS/IOSTargetPlatform/Private/IOSTargetPlatform.cpp b/Engine/Source/Developer/iOS/IOSTargetPlatform/Private/IOSTargetPlatform.cpp
--- a/Engine/Source/Developer/iOS/IOSTargetPlatform/Private/IOSTargetPlatform.cpp
+++ b/Engine/Source/Developer/iOS/IOSTargetPlatform/Private/IOSTargetPlatform.cpp
@@ -360,6 +360,24 @@ void FIOSTargetPlatform::GetAllPossibleShaderFormats( TArray<FName>& OutFormats
 	}
 }
 
+bool FIOSTargetPlatform::SupportsFeature( ETargetPlatformFeatures Feature ) const
+{
+	switch (Feature)
+	{
+		case ETargetPlatformFeatures::Packaging:
+			return true;
+
+		case ETargetPlatformFeatures::LowQualityLightmaps:
+			return SupportsES2() || SupportsMetal();
+		
+		case ETargetPlatformFeatures::HighQualityLightmaps:
+			return SupportsMetalMRT();
+		default:
+			break;
+	}
+	
+	return TTargetPlatformBase<FIOSPlatformProperties>::SupportsFeature(Feature);
+}
 
 void FIOSTargetPlatform::GetAllTargetedShaderFormats( TArray<FName>& OutFormats ) const
 {
diff --git a/Engine/Source/Developer/iOS/IOSTargetPlatform/Private/IOSTargetPlatform.h b/Engine/Source/Developer/iOS/IOSTargetPlatform/Private/IOSTargetPlatform.h
--- a/Engine/Source/Developer/iOS/IOSTargetPlatform/Private/IOSTargetPlatform.h
+++ b/Engine/Source/Developer/iOS/IOSTargetPlatform/Private/IOSTargetPlatform.h
@@ -72,16 +72,7 @@ public:
 		#endif
 	}
 
-	virtual bool SupportsFeature( ETargetPlatformFeatures Feature ) const override
-	{
-		if (Feature == ETargetPlatformFeatures::Packaging)
-		{
-			// not implemented yet
-			return true;
-		}
-
-		return TTargetPlatformBase<FIOSPlatformProperties>::SupportsFeature(Feature);
-	}
+	virtual bool SupportsFeature( ETargetPlatformFeatures Feature ) const override;
 
 	virtual bool IsSdkInstalled(bool bProjectHasCode, FString& OutTutorialPath) const override;
 	virtual int32 CheckRequirements(const FString& ProjectPath, bool bProjectHasCode, FString& OutTutorialPath) const override;
diff --git a/Engine/Source/Programs/AutomationTool/Android/AndroidPlatform.Automation.cs b/Engine/Source/Programs/AutomationTool/Android/AndroidPlatform.Automation.cs
--- a/Engine/Source/Programs/AutomationTool/Android/AndroidPlatform.Automation.cs
+++ b/Engine/Source/Programs/AutomationTool/Android/AndroidPlatform.Automation.cs
@@ -746,9 +746,9 @@ public class AndroidPlatform : Platform
 				DeviceArch = "-armv7";
 			}
 			// go from 64 to 32-bit
-			else if (DeviceArch == "-x86_64")
+			else if (DeviceArch == "-x64")
 			{
-				if (Array.IndexOf(AppArchitectures, "-x86") == -1)
+				if (Array.IndexOf(AppArchitectures, "-x86") != -1)
 				{
 					DeviceArch = "-x86";
 				}
@@ -774,7 +774,7 @@ public class AndroidPlatform : Platform
 		// if after the fallbacks, we still don't have it, we can't continue
 		if (Array.IndexOf(AppArchitectures, DeviceArch) == -1)
 		{
-            string ErrorString = String.Format("Unable to run because you don't have an apk that is usable on {0}", Params.Device);
+			string ErrorString = String.Format("Unable to run because you don't have an apk that is usable on {0}. Looked for {1}", Params.Device, DeviceArch);
             ErrorReporter.Error(ErrorString, (int)ErrorCodes.Error_NoApkSuitableForArchitecture);
             throw new AutomationException(ErrorString);
 		}
diff --git a/Engine/Source/Programs/UnrealBuildTool/Android/AndroidToolChain.cs b/Engine/Source/Programs/UnrealBuildTool/Android/AndroidToolChain.cs
--- a/Engine/Source/Programs/UnrealBuildTool/Android/AndroidToolChain.cs
+++ b/Engine/Source/Programs/UnrealBuildTool/Android/AndroidToolChain.cs
@@ -57,7 +57,7 @@ namespace UnrealBuildTool
 			}
 			if (Ini.GetBool("/Script/AndroidRuntimeSettings.AndroidRuntimeSettings", "bBuildForx8664", out bBuild) && bBuild)
 			{
-				ProjectArches.Add("-x86_64");
+				ProjectArches.Add("-x64");
 			}
 
 			// force armv7 if something went wrong
diff --git a/Engine/Source/Programs/UnrealBuildTool/Android/UEBuildAndroid.cs b/Engine/Source/Programs/UnrealBuildTool/Android/UEBuildAndroid.cs
--- a/Engine/Source/Programs/UnrealBuildTool/Android/UEBuildAndroid.cs
+++ b/Engine/Source/Programs/UnrealBuildTool/Android/UEBuildAndroid.cs
@@ -320,10 +320,7 @@ namespace UnrealBuildTool
 							InModule.AddPlatformSpecificDynamicallyLoadedModule("Android_DXTTargetPlatform");
 							InModule.AddPlatformSpecificDynamicallyLoadedModule("Android_ETC1TargetPlatform");
 							InModule.AddPlatformSpecificDynamicallyLoadedModule("Android_ETC2TargetPlatform");
-							InModule.AddPlatformSpecificDynamicallyLoadedModule("Android_ES31TargetPlatform");
 							InModule.AddPlatformSpecificDynamicallyLoadedModule("Android_ASTCTargetPlatform");
-							// @todo gl4android				
-							// InModule.AddPlatformSpecificDynamicallyLoadedModule("Android_GL4TargetPlatform");
                         }
 					}
 					else if (InModule.ToString() == "TargetPlatform")
@@ -349,10 +346,7 @@ namespace UnrealBuildTool
 					InModule.AddPlatformSpecificDynamicallyLoadedModule("Android_DXTTargetPlatform");
 					InModule.AddPlatformSpecificDynamicallyLoadedModule("Android_ETC1TargetPlatform");
 					InModule.AddPlatformSpecificDynamicallyLoadedModule("Android_ETC2TargetPlatform");
-					InModule.AddPlatformSpecificDynamicallyLoadedModule("Android_ES31TargetPlatform");
 					InModule.AddPlatformSpecificDynamicallyLoadedModule("Android_ASTCTargetPlatform");
-					// @todo gl4android				
-					// InModule.AddPlatformSpecificDynamicallyLoadedModule("Android_GL4TargetPlatform");
                 }
 
 				if (bBuildShaderFormats)
diff --git a/Engine/Source/Programs/UnrealBuildTool/Android/UEDeployAndroid.cs b/Engine/Source/Programs/UnrealBuildTool/Android/UEDeployAndroid.cs
--- a/Engine/Source/Programs/UnrealBuildTool/Android/UEDeployAndroid.cs
+++ b/Engine/Source/Programs/UnrealBuildTool/Android/UEDeployAndroid.cs
@@ -277,8 +277,10 @@ namespace UnrealBuildTool.Android
 		{
 			switch (UE4Arch)
 			{
-				case "-armv7": return "armeabi-v7a";
-				case "-x86": return "x86";
+				case "-armv7":	return "armeabi-v7a";
+				case "-arm64":	return "arm64";
+				case "-x64":	return "x86_64";
+				case "-x86":	return "x86";
 
 				default: throw new BuildException("Unknown UE4 architecture {0}", UE4Arch);
 			}
@@ -290,7 +292,10 @@ namespace UnrealBuildTool.Android
 			{
 				case "armeabi-v7a": return "-armv7";
 				case "x86":			return "-x86";
-
+				case "arm64":		return "-arm64";
+				case "x86_64":
+				case "x64":			return "-x64";
+					
 				default: throw new BuildException("Unknown NDK architecture '{0}'", NDKArch);
 			}
 		}
diff --git a/Engine/Source/Runtime/Core/Public/Android/AndroidProperties.h b/Engine/Source/Runtime/Core/Public/Android/AndroidProperties.h
--- a/Engine/Source/Runtime/Core/Public/Android/AndroidProperties.h
+++ b/Engine/Source/Runtime/Core/Public/Android/AndroidProperties.h
@@ -43,15 +43,6 @@ struct FAndroidPlatformProperties
 		return UE_GAME;
 	}
 
-	static FORCEINLINE uint32 MaxGpuSkinBones( )
-	{
-#if PLATFORM_ANDROIDES31 || PLATFORM_ANDROIDGL4
-		return 256;
-#else
-		return 20;
-#endif
-	}
-
 	static FORCEINLINE bool RequiresCookedData( )
 	{
 		return true;
diff --git a/Engine/Source/Runtime/Core/Public/GenericPlatform/GenericPlatformProperties.h b/Engine/Source/Runtime/Core/Public/GenericPlatform/GenericPlatformProperties.h
--- a/Engine/Source/Runtime/Core/Public/GenericPlatform/GenericPlatformProperties.h
+++ b/Engine/Source/Runtime/Core/Public/GenericPlatform/GenericPlatformProperties.h
@@ -100,16 +100,6 @@ struct FGenericPlatformProperties
 		return true;
 	}
 
-	/**
-	 * Returns the maximum bones the platform supports.
-	 *
-	 * @return the maximum bones the platform supports.
-	 */
-	static FORCEINLINE uint32 MaxGpuSkinBones()
-	{
-		return 256;
-	}
-
 	/**
 	 * Gets the name of this platform
 	 *
diff --git a/Engine/Source/Runtime/Core/Public/HTML5/HTML5PlatformProperties.h b/Engine/Source/Runtime/Core/Public/HTML5/HTML5PlatformProperties.h
--- a/Engine/Source/Runtime/Core/Public/HTML5/HTML5PlatformProperties.h
+++ b/Engine/Source/Runtime/Core/Public/HTML5/HTML5PlatformProperties.h
@@ -28,11 +28,6 @@ struct FHTML5PlatformProperties
 		return false;
 	}
 
-	static FORCEINLINE uint32 MaxGpuSkinBones( )
-	{
-		return 20;
-	}
-
 	static FORCEINLINE const char* PlatformName()
 	{
 		return "HTML5";
diff --git a/Engine/Source/Runtime/OpenGLDrv/Private/Android/AndroidES31OpenGL.h b/Engine/Source/Runtime/OpenGLDrv/Private/Android/AndroidES31OpenGL.h
--- a/Engine/Source/Runtime/OpenGLDrv/Private/Android/AndroidES31OpenGL.h
+++ b/Engine/Source/Runtime/OpenGLDrv/Private/Android/AndroidES31OpenGL.h
@@ -64,6 +64,13 @@ typedef FEGLGLSync UGLsync;
 #define EGL_OPENGL_ES3_BIT_KHR                             0x40       /* 64 */
 #endif
 
+// _EXT was needed for Android, so make both work if non _EXT isn't defined
+#ifndef GL_TESS_CONTROL_SHADER_BIT
+#define GL_TESS_CONTROL_SHADER_BIT GL_TESS_CONTROL_SHADER_BIT_EXT
+#define GL_TESS_EVALUATION_SHADER_BIT GL_TESS_EVALUATION_SHADER_BIT_EXT
+#define GL_GEOMETRY_SHADER_BIT GL_GEOMETRY_SHADER_BIT_EXT
+#endif
+
 #ifndef GL_NV_bindless_texture
 #define GL_NV_bindless_texture 1
 typedef uint64_t GLuint64EXT;
