- Fixed up tooltips/comments for new ES31+AEP support

diff --git a/Engine/Source/Developer/Android/AndroidPlatformEditor/Private/AndroidTargetSettingsCustomization.cpp b/Engine/Source/Developer/Android/AndroidPlatformEditor/Private/AndroidTargetSettingsCustomization.cpp
--- a/Engine/Source/Developer/Android/AndroidPlatformEditor/Private/AndroidTargetSettingsCustomization.cpp
+++ b/Engine/Source/Developer/Android/AndroidPlatformEditor/Private/AndroidTargetSettingsCustomization.cpp
@@ -202,7 +202,7 @@ void FAndroidTargetSettingsCustomization::BuildAppManifestSection(IDetailLayoutB
 	SETUP_NONROCKET_PROP(bBuildForArmV7, BuildCategory, LOCTEXT("BuildForArmV7ToolTip", "Enable ArmV7 CPU architecture support? (this will be used if all CPU architecture types are unchecked)"));
 	SETUP_NONROCKET_PROP(bBuildForX86, BuildCategory, LOCTEXT("BuildForX86ToolTip", "Enable X86 CPU architecture support?"));
 	SETUP_NONROCKET_PROP(bBuildForES2, BuildCategory, LOCTEXT("BuildForES2ToolTip", "Enable OpenGL ES2 rendering support? (this will be used if rendering types are unchecked)"));
-	SETUP_NONROCKET_PROP(bBuildForES31, BuildCategory, LOCTEXT("BuildForES31ToolTip", "Enable OpenGL ES31 + AEP (Android Extension Pack) rendering support? Currently only Tegra K1 supports this, as it will force DXT textures (In 4.8 3.1+AEP will work with all texture formats).\nIf you use the Launch On feature (in the main toolbar), when you change this setting, you need to restart the editor to make sure it will launch with the proper 3.1+AEP support!"));
+	SETUP_NONROCKET_PROP(bBuildForES31, BuildCategory, LOCTEXT("BuildForES31ToolTip", "Enable OpenGL ES31 + AEP (Android Extension Pack) rendering support?"));
 	
 	// @todo android fat binary: Put back in when we expose those
 //	SETUP_NONROCKET_PROP(bSplitIntoSeparateApks, BuildCategory, LOCTEXT("SplitIntoSeparateAPKsToolTip", "If checked, CPU architectures and rendering types will be split into separate .apk files"));
diff --git a/Engine/Source/Programs/AutomationTool/Android/AndroidPlatform.Automation.cs b/Engine/Source/Programs/AutomationTool/Android/AndroidPlatform.Automation.cs
--- a/Engine/Source/Programs/AutomationTool/Android/AndroidPlatform.Automation.cs
+++ b/Engine/Source/Programs/AutomationTool/Android/AndroidPlatform.Automation.cs
@@ -935,7 +935,7 @@ public class AndroidPlatform : Platform
 		ProcessResult ExtensionsResult = RunAdbCommand(Params, "shell dumpsys SurfaceFlinger", null, ERunOptions.AppMustExist);
 		string Extensions = ExtensionsResult.Output.Trim();
 
-		// look for AEP support
+		// look for AEP support (on device and in project)
 		if (Extensions.Contains("GL_ANDROID_extension_pack_es31a") && Extensions.Contains("GL_EXT_color_buffer_half_float"))
 		{
 			if (AppGPUArchitectures.Contains("-es31"))
diff --git a/Engine/Source/Runtime/Android/AndroidRuntimeSettings/Classes/AndroidRuntimeSettings.h b/Engine/Source/Runtime/Android/AndroidRuntimeSettings/Classes/AndroidRuntimeSettings.h
--- a/Engine/Source/Runtime/Android/AndroidRuntimeSettings/Classes/AndroidRuntimeSettings.h
+++ b/Engine/Source/Runtime/Android/AndroidRuntimeSettings/Classes/AndroidRuntimeSettings.h
@@ -185,7 +185,7 @@ public:
 	bool bBuildForES2;
 
 	// Enable ES31 support? [CURRENTLY FOR FULL SOURCE GAMES ONLY]
-	UPROPERTY(GlobalConfig, EditAnywhere, Category = Build, meta = (DisplayName = "Support OpenGL ES31 + AEP (restart may be needed - see tooltip)"))
+	UPROPERTY(GlobalConfig, EditAnywhere, Category = Build, meta = (DisplayName = "Support OpenGL ES31 + AEP"))
 	bool bBuildForES31;
 
 	// If selected, the checked architectures will be split into separate .apk files [CURRENTLY FOR FULL SOURCE GAMES ONLY]
