- Added the non-POT error texture to Android when using PVRTC [UE-8923]

diff --git a/Engine/Source/Developer/Android/AndroidTargetPlatform/Private/AndroidTargetPlatform.h b/Engine/Source/Developer/Android/AndroidTargetPlatform/Private/AndroidTargetPlatform.h
--- a/Engine/Source/Developer/Android/AndroidTargetPlatform/Private/AndroidTargetPlatform.h
+++ b/Engine/Source/Developer/Android/AndroidTargetPlatform/Private/AndroidTargetPlatform.h
@@ -45,6 +45,9 @@ namespace AndroidTexFormat
 	static FName NameG8(TEXT("G8"));
 	static FName NameVU8(TEXT("VU8"));
 	static FName NameRGBA16F(TEXT("RGBA16F"));
+
+	// Error "formats" (uncompressed)
+	static FName NamePOTERROR(TEXT("POTERROR"));
 }
 
 
@@ -120,6 +123,12 @@ public:
 		return true;
 	}
 
+	virtual bool SupportsCompressedNonPOT( ) const
+	{
+		// most formats do support non-POT compressed textures
+		return true;
+	}
+
 #if WITH_ENGINE
 	virtual void GetReflectionCaptureFormats( TArray<FName>& OutFormats ) const override;
 
@@ -173,8 +182,9 @@ protected:
 	 *
 	 * @param Format - The format to add.
 	 * @param OutFormats - The collection of formats to add to.
+	 * @param bIsCompressedNonPOT - If this is true, the texture wants to be compressed but is not a power of 2
 	 */
-	void AddTextureFormatIfSupports( FName Format, TArray<FName>& OutFormats ) const;
+	void AddTextureFormatIfSupports( FName Format, TArray<FName>& OutFormats, bool bIsCompressedNonPOT=false ) const;
 
 	/**
 	 * Return true if this device has a supported set of extensions for this platform.
diff --git a/Engine/Source/Developer/Android/AndroidTargetPlatform/Private/AndroidTargetPlatform.inl b/Engine/Source/Developer/Android/AndroidTargetPlatform/Private/AndroidTargetPlatform.inl
--- a/Engine/Source/Developer/Android/AndroidTargetPlatform/Private/AndroidTargetPlatform.inl
+++ b/Engine/Source/Developer/Android/AndroidTargetPlatform/Private/AndroidTargetPlatform.inl
@@ -189,7 +189,15 @@ inline void FAndroidTargetPlatform<TPlatformProperties>::GetTextureFormats( cons
 		|| (InTexture->Source.GetSizeY() < 4)
 		|| (InTexture->Source.GetSizeX() % 4 != 0)
 		|| (InTexture->Source.GetSizeY() % 4 != 0);
-
+	
+	bool bIsNonPOT = false;
+#if WITH_EDITORONLY_DATA
+	// is this texture not a power of 2?
+	bIsNonPOT = !InTexture->Source.IsPowerOfTwo();
+	
+	if (bIsNonPOT) UE_LOG(LogTemp, Display, TEXT("NonPOT texture, %d x %d"), InTexture->Source.GetSizeX(), InTexture->Source.GetSizeY());
+#endif
+	
 	// Determine the pixel format of the compressed texture.
 	if (bNoCompression && InTexture->HasHDRSource())
 	{
@@ -206,11 +214,11 @@ inline void FAndroidTargetPlatform<TPlatformProperties>::GetTextureFormats( cons
 	}
 	else if (InTexture->CompressionSettings == TC_Normalmap)
 	{
-		AddTextureFormatIfSupports(AndroidTexFormat::NamePVRTC4, OutFormats);
-		AddTextureFormatIfSupports(AndroidTexFormat::NameDXT5, OutFormats);
-		AddTextureFormatIfSupports(AndroidTexFormat::NameATC_RGBA_I, OutFormats);
-		AddTextureFormatIfSupports(AndroidTexFormat::NameAutoETC2, OutFormats);
-		AddTextureFormatIfSupports(AndroidTexFormat::NameAutoETC1, OutFormats);
+		AddTextureFormatIfSupports(AndroidTexFormat::NamePVRTC4, OutFormats, bIsNonPOT);
+		AddTextureFormatIfSupports(AndroidTexFormat::NameDXT5, OutFormats, bIsNonPOT);
+		AddTextureFormatIfSupports(AndroidTexFormat::NameATC_RGBA_I, OutFormats, bIsNonPOT);
+		AddTextureFormatIfSupports(AndroidTexFormat::NameAutoETC2, OutFormats, bIsNonPOT);
+		AddTextureFormatIfSupports(AndroidTexFormat::NameAutoETC1, OutFormats, bIsNonPOT);
 	}
 	else if (InTexture->CompressionSettings == TC_Displacementmap)
 	{
@@ -235,35 +243,35 @@ inline void FAndroidTargetPlatform<TPlatformProperties>::GetTextureFormats( cons
 	else if (InTexture->bForcePVRTC4
 		|| InTexture->CompressionSettings == TC_BC7)
 	{
-		AddTextureFormatIfSupports(AndroidTexFormat::NamePVRTC4, OutFormats);
-		AddTextureFormatIfSupports(AndroidTexFormat::NameDXT5, OutFormats);
-		AddTextureFormatIfSupports(AndroidTexFormat::NameATC_RGBA_I, OutFormats);
-		AddTextureFormatIfSupports(AndroidTexFormat::NameAutoETC2, OutFormats);
-		AddTextureFormatIfSupports(AndroidTexFormat::NameAutoETC1, OutFormats);
+		AddTextureFormatIfSupports(AndroidTexFormat::NamePVRTC4, OutFormats, bIsNonPOT);
+		AddTextureFormatIfSupports(AndroidTexFormat::NameDXT5, OutFormats, bIsNonPOT);
+		AddTextureFormatIfSupports(AndroidTexFormat::NameATC_RGBA_I, OutFormats, bIsNonPOT);
+		AddTextureFormatIfSupports(AndroidTexFormat::NameAutoETC2, OutFormats, bIsNonPOT);
+		AddTextureFormatIfSupports(AndroidTexFormat::NameAutoETC1, OutFormats, bIsNonPOT);
 	}
 	else if (InTexture->CompressionNoAlpha)
 	{
-		AddTextureFormatIfSupports(AndroidTexFormat::NamePVRTC2, OutFormats);
-		AddTextureFormatIfSupports(AndroidTexFormat::NameDXT1, OutFormats);
-		AddTextureFormatIfSupports(AndroidTexFormat::NameATC_RGB, OutFormats);
-		AddTextureFormatIfSupports(AndroidTexFormat::NameETC2_RGB, OutFormats);
-		AddTextureFormatIfSupports(AndroidTexFormat::NameETC1, OutFormats);
+		AddTextureFormatIfSupports(AndroidTexFormat::NamePVRTC2, OutFormats, bIsNonPOT);
+		AddTextureFormatIfSupports(AndroidTexFormat::NameDXT1, OutFormats, bIsNonPOT);
+		AddTextureFormatIfSupports(AndroidTexFormat::NameATC_RGB, OutFormats, bIsNonPOT);
+		AddTextureFormatIfSupports(AndroidTexFormat::NameETC2_RGB, OutFormats, bIsNonPOT);
+		AddTextureFormatIfSupports(AndroidTexFormat::NameETC1, OutFormats, bIsNonPOT);
 	}
 	else if (InTexture->bDitherMipMapAlpha)
 	{
-		AddTextureFormatIfSupports(AndroidTexFormat::NamePVRTC4, OutFormats);
-		AddTextureFormatIfSupports(AndroidTexFormat::NameDXT5, OutFormats);
-		AddTextureFormatIfSupports(AndroidTexFormat::NameATC_RGBA_I, OutFormats);
-		AddTextureFormatIfSupports(AndroidTexFormat::NameAutoETC2, OutFormats);
-		AddTextureFormatIfSupports(AndroidTexFormat::NameAutoETC1, OutFormats);
+		AddTextureFormatIfSupports(AndroidTexFormat::NamePVRTC4, OutFormats, bIsNonPOT);
+		AddTextureFormatIfSupports(AndroidTexFormat::NameDXT5, OutFormats, bIsNonPOT);
+		AddTextureFormatIfSupports(AndroidTexFormat::NameATC_RGBA_I, OutFormats, bIsNonPOT);
+		AddTextureFormatIfSupports(AndroidTexFormat::NameAutoETC2, OutFormats, bIsNonPOT);
+		AddTextureFormatIfSupports(AndroidTexFormat::NameAutoETC1, OutFormats, bIsNonPOT);
 	}
 	else
 	{
-		AddTextureFormatIfSupports(AndroidTexFormat::NameAutoPVRTC, OutFormats);
-		AddTextureFormatIfSupports(AndroidTexFormat::NameAutoDXT, OutFormats);
-		AddTextureFormatIfSupports(AndroidTexFormat::NameAutoATC, OutFormats);
-		AddTextureFormatIfSupports(AndroidTexFormat::NameAutoETC2, OutFormats);
-		AddTextureFormatIfSupports(AndroidTexFormat::NameAutoETC1, OutFormats);
+		AddTextureFormatIfSupports(AndroidTexFormat::NameAutoPVRTC, OutFormats, bIsNonPOT);
+		AddTextureFormatIfSupports(AndroidTexFormat::NameAutoDXT, OutFormats, bIsNonPOT);
+		AddTextureFormatIfSupports(AndroidTexFormat::NameAutoATC, OutFormats, bIsNonPOT);
+		AddTextureFormatIfSupports(AndroidTexFormat::NameAutoETC2, OutFormats, bIsNonPOT);
+		AddTextureFormatIfSupports(AndroidTexFormat::NameAutoETC1, OutFormats, bIsNonPOT);
 	}
 }
 
@@ -346,11 +354,18 @@ FText FAndroidTargetPlatform<TPlatformProperties>::GetVariantTitle() const
  *****************************************************************************/
 
 template<class TPlatformProperties>
-inline void FAndroidTargetPlatform<TPlatformProperties>::AddTextureFormatIfSupports( FName Format, TArray<FName>& OutFormats ) const
+inline void FAndroidTargetPlatform<TPlatformProperties>::AddTextureFormatIfSupports( FName Format, TArray<FName>& OutFormats, bool bIsCompressedNonPOT ) const
 {
 	if (SupportsTextureFormat(Format))
 	{
-		OutFormats.Add(Format);
+		if (bIsCompressedNonPOT && SupportsCompressedNonPOT() == false)
+		{
+			OutFormats.Add(AndroidTexFormat::NamePOTERROR);
+		}
+		else
+		{
+			OutFormats.Add(Format);
+		}
 	}
 }
 
diff --git a/Engine/Source/Developer/Android/Android_PVRTCTargetPlatform/Private/Android_PVRTCTargetPlatformModule.cpp b/Engine/Source/Developer/Android/Android_PVRTCTargetPlatform/Private/Android_PVRTCTargetPlatformModule.cpp
--- a/Engine/Source/Developer/Android/Android_PVRTCTargetPlatform/Private/Android_PVRTCTargetPlatformModule.cpp
+++ b/Engine/Source/Developer/Android/Android_PVRTCTargetPlatform/Private/Android_PVRTCTargetPlatformModule.cpp
@@ -29,6 +29,11 @@ class FAndroid_PVRTCTargetPlatform
 		return FString(FAndroid_PVRTCPlatformProperties::PlatformName());
 	}
 
+	virtual bool SupportsCompressedNonPOT( ) const override
+	{
+		return false;
+	}
+	
 	virtual bool SupportsTextureFormat( FName Format ) const override
 	{
 		if( Format == AndroidTexFormat::NamePVRTC2 ||
