Fix typo

diff --git a/Engine/Source/Developer/Android/AndroidPlatformEditor/Private/AndroidTargetSettingsCustomization.cpp b/Engine/Source/Developer/Android/AndroidPlatformEditor/Private/AndroidTargetSettingsCustomization.cpp
--- a/Engine/Source/Developer/Android/AndroidPlatformEditor/Private/AndroidTargetSettingsCustomization.cpp
+++ b/Engine/Source/Developer/Android/AndroidPlatformEditor/Private/AndroidTargetSettingsCustomization.cpp
@@ -46,7 +46,7 @@ FAndroidTargetSettingsCustomization::FAndroidTargetSettingsCustomization()
 	new (IconNames) FPlatformIconInfo(TEXT("res/drawable-mdpi/icon.png"), LOCTEXT("SettingsIcon_MDPI", "MDPI Icon"), FText::GetEmpty(), 48, 48, FPlatformIconInfo::Required);
 	new (IconNames) FPlatformIconInfo(TEXT("res/drawable-hdpi/icon.png"), LOCTEXT("SettingsIcon_HDPI", "HDPI Icon"), FText::GetEmpty(), 72, 72, FPlatformIconInfo::Required);
 	new (IconNames) FPlatformIconInfo(TEXT("res/drawable-xhdpi/icon.png"), LOCTEXT("SettingsIcon_XHDPI", "XHDPI Icon"), FText::GetEmpty(), 96, 96, FPlatformIconInfo::Required);
-	new (IconNames) FPlatformIconInfo(TEXT("res/drawable/downloadimagev.png"), LOCTEXT("SettingsIcon_DownloadImageV", "Download Background Verticle Image"), FText::GetEmpty(), 720, 1280, FPlatformIconInfo::Required);
+	new (IconNames) FPlatformIconInfo(TEXT("res/drawable/downloadimagev.png"), LOCTEXT("SettingsIcon_DownloadImageV", "Download Background Vertical Image"), FText::GetEmpty(), 720, 1280, FPlatformIconInfo::Required);
 }
 
 void FAndroidTargetSettingsCustomization::CustomizeDetails(IDetailLayoutBuilder& DetailLayout)
