Removed log spam

diff --git a/Engine/Source/Developer/Android/AndroidTargetPlatform/Private/AndroidTargetPlatform.inl b/Engine/Source/Developer/Android/AndroidTargetPlatform/Private/AndroidTargetPlatform.inl
--- a/Engine/Source/Developer/Android/AndroidTargetPlatform/Private/AndroidTargetPlatform.inl
+++ b/Engine/Source/Developer/Android/AndroidTargetPlatform/Private/AndroidTargetPlatform.inl
@@ -194,8 +194,6 @@ inline void FAndroidTargetPlatform<TPlatformProperties>::GetTextureFormats( cons
 #if WITH_EDITORONLY_DATA
 	// is this texture not a power of 2?
 	bIsNonPOT = !InTexture->Source.IsPowerOfTwo();
-	
-	if (bIsNonPOT) UE_LOG(LogTemp, Display, TEXT("NonPOT texture, %d x %d"), InTexture->Source.GetSizeX(), InTexture->Source.GetSizeY());
 #endif
 	
 	// Determine the pixel format of the compressed texture.
