UEMOB-32 - Support cooking various combinations of Android texture formats simultaneously

diff --git a/Engine/Config/BaseEngine.ini b/Engine/Config/BaseEngine.ini
--- a/Engine/Config/BaseEngine.ini
+++ b/Engine/Config/BaseEngine.ini
@@ -2322,6 +2322,12 @@ StoreVersion=1
 VersionDisplayName=1.0
 MinSDKVersion=9
 bShowLaunchImage=true
+bMultiTargetFormat_ETC1=true
+bMultiTargetFormat_ETC2=true
+bMultiTargetFormat_DXT=true
+bMultiTargetFormat_PVRTC=true
+bMultiTargetFormat_ATC=true
+bMultiTargetFormat_ASTC=true
 
 [/Script/AndroidPlatformEditor.AndroidSDKSettings]
 SDKAPILevel=latest
diff --git a/Engine/Source/Developer/Android/Android_MultiTargetPlatform/Android_MultiTargetPlatform.Build.cs b/Engine/Source/Developer/Android/Android_MultiTargetPlatform/Android_MultiTargetPlatform.Build.cs
new file mode 100644
--- /dev/null
+++ b/Engine/Source/Developer/Android/Android_MultiTargetPlatform/Android_MultiTargetPlatform.Build.cs
@@ -0,0 +1,40 @@
+// Copyright 1998-2015 Epic Games, Inc. All Rights Reserved.
+
+using UnrealBuildTool;
+
+public class Android_MultiTargetPlatform : ModuleRules
+{
+	public Android_MultiTargetPlatform( TargetInfo Target )
+	{
+		BinariesSubFolder = "Android";
+
+        PrivateDependencyModuleNames.AddRange(
+			new string[]
+			{
+				"Core",
+				"TargetPlatform",
+				"DesktopPlatform",
+				"AndroidDeviceDetection",
+			}
+		);
+
+		PublicIncludePaths.AddRange(
+			new string[]
+			{
+				"Runtime/Core/Public/Android"
+			}
+		);
+
+		if (UEBuildConfiguration.bCompileAgainstEngine)
+		{
+			PrivateDependencyModuleNames.Add("Engine");
+		}
+
+		PrivateIncludePaths.AddRange(
+			new string[]
+			{
+				"Developer/Android/AndroidTargetPlatform/Private",				
+			}
+		);
+	}
+}
\ No newline at end of file
diff --git a/Engine/Source/Developer/Android/Android_MultiTargetPlatform/Private/Android_MultiTargetPlatformModule.cpp b/Engine/Source/Developer/Android/Android_MultiTargetPlatform/Private/Android_MultiTargetPlatformModule.cpp
new file mode 100644
--- /dev/null
+++ b/Engine/Source/Developer/Android/Android_MultiTargetPlatform/Private/Android_MultiTargetPlatformModule.cpp
@@ -0,0 +1,160 @@
+// Copyright 1998-2015 Epic Games, Inc. All Rights Reserved.
+
+/*=============================================================================
+	Android_MultiTargetPlatformModule.cpp: Implements the Android_Multi target platform variant
+=============================================================================*/
+
+#include "Android_MultiTargetPlatformPrivatePCH.h"
+
+#define LOCTEXT_NAMESPACE "FAndroid_MultiTargetPlatformModule" 
+
+/**
+ * Android cooking platform which cooks multiple texture formats
+ */
+class FAndroid_MultiTargetPlatform
+	: public FAndroidTargetPlatform<FAndroid_MultiPlatformProperties>
+{
+	TArray<ITargetPlatform*> FormatTargetPlatforms;
+	FString FormatTargetString;
+
+	FAndroid_MultiTargetPlatform()
+	{
+		LoadFormats();
+	}
+
+	void LoadFormats()
+	{
+		const TCHAR* FormatNames[] =
+		{
+			TEXT("ASTC"),
+			TEXT("ATC"),
+			TEXT("PVRTC"),
+			TEXT("DXT"),
+			TEXT("ETC2"),
+			TEXT("ETC1"),
+		};
+
+		FormatTargetPlatforms.Empty();
+		FormatTargetString = TEXT("");
+
+		// Load the TargetPlatform module for each format
+		for (int32 PlatformFormatIdx = 0; PlatformFormatIdx < ARRAY_COUNT(FormatNames); PlatformFormatIdx++)
+		{
+			bool bEnabled = false;
+			FString SettingsName = FString(TEXT("bMultiTargetFormat_")) + FormatNames[PlatformFormatIdx];
+			GConfig->GetBool(TEXT("/Script/AndroidRuntimeSettings.AndroidRuntimeSettings"), *SettingsName, bEnabled, GEngineIni);
+			if (bEnabled)
+			{
+				FString PlatformName = FString(TEXT("Android_")) + FormatNames[PlatformFormatIdx] + TEXT("TargetPlatform");
+				ITargetPlatformModule* Module = FModuleManager::LoadModulePtr<ITargetPlatformModule>(*PlatformName);
+				if (Module)
+				{
+					ITargetPlatform* TargetPlatform = Module->GetTargetPlatform();
+					if (TargetPlatform)
+					{
+						if (FormatTargetPlatforms.Num())
+						{
+							FormatTargetString += TEXT(",");
+						}
+						FormatTargetString += FormatNames[PlatformFormatIdx];
+						FormatTargetPlatforms.Add(TargetPlatform);
+					}
+				}
+			}
+		}
+
+		PlatformInfo::UpdatePlatformDisplayName(TEXT("Android_Multi"), DisplayName());
+	}
+
+	virtual FString GetAndroidVariantName() override
+	{
+		return TEXT("Multi");
+	}
+
+	virtual FText DisplayName() const override
+	{
+		return FText::Format(LOCTEXT("Android_Multi", "Android (Multi:{0})"), FText::FromString(FormatTargetString));
+	}
+
+	virtual FString PlatformName() const override
+	{
+		return FString(FAndroid_MultiPlatformProperties::PlatformName());
+	}
+
+	virtual void GetTextureFormats(const UTexture* Texture, TArray<FName>& OutFormats) const
+	{
+		// Ask each platform variant to choose texture formats
+		for (ITargetPlatform* Platform : FormatTargetPlatforms)
+		{
+			TArray<FName> PlatformFormats;
+			Platform->GetTextureFormats(Texture, PlatformFormats);
+			for (FName Format : PlatformFormats)
+			{
+				OutFormats.AddUnique(Format);
+			}
+		}
+	}
+	
+	virtual FText GetVariantDisplayName() const override
+	{
+		return LOCTEXT("Android_Multi_ShortName", "Multi");
+	}
+
+	virtual float GetVariantPriority() const override
+	{
+		// lowest priority so specific variants are chosen first
+		return 0.f;
+	}
+
+	// for constructor
+	friend class FAndroid_MultiTargetPlatformModule;
+};
+
+/**
+ * Holds the target platform singleton.
+ */
+static ITargetPlatform* AndroidTargetSingleton = NULL;
+
+
+/**
+ * Module for the Android target platform.
+ */
+class FAndroid_MultiTargetPlatformModule : public IAndroid_MultiTargetPlatformModule
+{
+public:
+
+	/**
+	 * Destructor.
+	 */
+	~FAndroid_MultiTargetPlatformModule( )
+	{
+		AndroidTargetSingleton = NULL;
+	}
+
+
+	// Begin ITargetPlatformModule interface
+	virtual ITargetPlatform* GetTargetPlatform() override
+	{
+		if (AndroidTargetSingleton == NULL)
+		{
+			AndroidTargetSingleton = new FAndroid_MultiTargetPlatform();
+		}
+		
+		return AndroidTargetSingleton;
+	}
+	// End ITargetPlatformModule interface
+
+	// Begin IAndroid_MultiTargetPlatformModule interface
+	virtual void NotifySelectedFormatsChanged() override
+	{
+		static_cast<FAndroid_MultiTargetPlatform*>(GetTargetPlatform())->LoadFormats();
+	}
+	// End IAndroid_MultiTargetPlatformModule interface
+
+};
+
+
+#undef LOCTEXT_NAMESPACE
+
+
+IMPLEMENT_MODULE( FAndroid_MultiTargetPlatformModule, Android_MultiTargetPlatform);
diff --git a/Engine/Source/Developer/Android/Android_MultiTargetPlatform/Private/Android_MultiTargetPlatformPrivatePCH.h b/Engine/Source/Developer/Android/Android_MultiTargetPlatform/Private/Android_MultiTargetPlatformPrivatePCH.h
new file mode 100644
--- /dev/null
+++ b/Engine/Source/Developer/Android/Android_MultiTargetPlatform/Private/Android_MultiTargetPlatformPrivatePCH.h
@@ -0,0 +1,32 @@
+// Copyright 1998-2015 Epic Games, Inc. All Rights Reserved.
+
+
+/*=============================================================================
+	Android_MultiTargetPlatformPrivatePCH.h: Pre-compiled header file for the Android_MultiTargetPlatform module.
+=============================================================================*/
+
+#pragma once
+
+/* Dependencies
+ *****************************************************************************/
+
+#include "Core.h"
+#include "ModuleInterface.h"
+#include "ModuleManager.h"
+#include "Runtime/Core/Public/Android/AndroidProperties.h"
+
+#if WITH_ENGINE
+	#include "Engine.h"
+#endif
+
+#include "PlatformInfo.h"
+#include "TargetPlatform.h"
+#include "TargetPlatformBase.h"
+#include "AndroidDeviceDetection.h"
+#include "IAndroid_MultiTargetPlatformModule.h"
+
+/* Private includes
+ *****************************************************************************/
+
+#include "AndroidTargetDevice.h"
+#include "AndroidTargetPlatform.h"
diff --git a/Engine/Source/Developer/Android/Android_MultiTargetPlatform/Public/IAndroid_MultiTargetPlatformModule.h b/Engine/Source/Developer/Android/Android_MultiTargetPlatform/Public/IAndroid_MultiTargetPlatformModule.h
new file mode 100644
--- /dev/null
+++ b/Engine/Source/Developer/Android/Android_MultiTargetPlatform/Public/IAndroid_MultiTargetPlatformModule.h
@@ -0,0 +1,21 @@
+// Copyright 1998-2015 Epic Games, Inc. All Rights Reserved.
+
+#pragma  once 
+
+/*------------------------------------------------------------------------------------
+IAndroid_MultiTargetPlatformModule interface 
+------------------------------------------------------------------------------------*/
+
+#include "Core.h"
+#include "ModuleInterface.h"
+#include "ITargetPlatformModule.h"
+
+class IAndroid_MultiTargetPlatformModule
+	: public ITargetPlatformModule
+{
+public:
+	//
+	// Called by AndroidRuntimeSettings to notify us when the user changes the selected texture formats
+	//
+	virtual void NotifySelectedFormatsChanged() = 0;
+};
\ No newline at end of file
diff --git a/Engine/Source/Developer/DesktopPlatform/Private/PlatformInfo.cpp b/Engine/Source/Developer/DesktopPlatform/Private/PlatformInfo.cpp
--- a/Engine/Source/Developer/DesktopPlatform/Private/PlatformInfo.cpp
+++ b/Engine/Source/Developer/DesktopPlatform/Private/PlatformInfo.cpp
@@ -98,6 +98,7 @@ static const FPlatformInfo PlatformInfoArray[] = {
 
 	BuildPlatformInfo(TEXT("Android"),					TEXT("Android"),			LOCTEXT("Android", "Android"),								EPlatformType::Game,		EPlatformFlags::None,			FPlatformIconPaths(TEXT("Launcher/AndroidTarget/Platform_Android_24x"), TEXT("Launcher/AndroidTarget/Platform_Android_128x")),						TEXT(""),												TEXT("Android"),	EPlatformSDKStatus::Unknown,	TEXT("/Engine/Tutorial/Mobile/SettingUpAndroidTutorial.SettingUpAndroidTutorial"),						IsAvailableOnWindows || IsAvailableOnMac,	true,				false,					true),
 	BuildPlatformInfo(TEXT("Android_All"),				TEXT("Android"),			LOCTEXT("Android_All", "Android (All)"),					EPlatformType::Game,		EPlatformFlags::CookFlavor,		FPlatformIconPaths(TEXT("Launcher/AndroidTarget/Platform_Android_24x"), TEXT("Launcher/AndroidTarget/Platform_Android_128x")),						TEXT(""),												TEXT("Android"),	EPlatformSDKStatus::Unknown,	TEXT("/Engine/Tutorial/Mobile/SettingUpAndroidTutorial.SettingUpAndroidTutorial"),						IsAvailableOnWindows || IsAvailableOnMac,	true,				false,					true),
+	BuildPlatformInfo(TEXT("Android_Multi"),			TEXT("Android_Multi"),		LOCTEXT("Android_Multi", "Android (Multi)"),				EPlatformType::Game,		EPlatformFlags::CookFlavor,		FPlatformIconPaths(TEXT("Launcher/AndroidTarget/Platform_Android_24x"), TEXT("Launcher/AndroidTarget/Platform_Android_128x")),						TEXT("-targetplatform=Android -cookflavor=Multi"),		TEXT("Android"),	EPlatformSDKStatus::Unknown,	TEXT("/Engine/Tutorial/Mobile/SettingUpAndroidTutorial.SettingUpAndroidTutorial"),						IsAvailableOnWindows || IsAvailableOnMac,	true,				false,					true),
 	BuildPlatformInfo(TEXT("Android_ATC"),				TEXT("Android_ATC"),		LOCTEXT("Android_ATC", "Android (ATC)"),					EPlatformType::Game,		EPlatformFlags::CookFlavor,		FPlatformIconPaths(TEXT("Launcher/AndroidTarget/Platform_Android_ATC_24x"), TEXT("Launcher/AndroidTarget/Platform_Android_128x")),					TEXT("-targetplatform=Android -cookflavor=ATC"),		TEXT("Android"),	EPlatformSDKStatus::Unknown,	TEXT("/Engine/Tutorial/Mobile/SettingUpAndroidTutorial.SettingUpAndroidTutorial"),						IsAvailableOnWindows || IsAvailableOnMac,	true,				false,					true),
 	BuildPlatformInfo(TEXT("Android_DXT"),				TEXT("Android_DXT"),		LOCTEXT("Android_DXT", "Android (DXT)"),					EPlatformType::Game,		EPlatformFlags::CookFlavor,		FPlatformIconPaths(TEXT("Launcher/AndroidTarget/Platform_Android_DXT_24x"), TEXT("Launcher/AndroidTarget/Platform_Android_128x")),					TEXT("-targetplatform=Android -cookflavor=DXT"),		TEXT("Android"),	EPlatformSDKStatus::Unknown,	TEXT("/Engine/Tutorial/Mobile/SettingUpAndroidTutorial.SettingUpAndroidTutorial"),						IsAvailableOnWindows || IsAvailableOnMac,	true,				false,					true),
 	BuildPlatformInfo(TEXT("Android_ETC1"),				TEXT("Android_ETC1"),		LOCTEXT("Android_ETC1", "Android (ETC1)"),					EPlatformType::Game,		EPlatformFlags::CookFlavor,		FPlatformIconPaths(TEXT("Launcher/AndroidTarget/Platform_Android_ETC1_24x"), TEXT("Launcher/AndroidTarget/Platform_Android_128x")),					TEXT("-targetplatform=Android -cookflavor=ETC1"),		TEXT("Android"),	EPlatformSDKStatus::Unknown,	TEXT("/Engine/Tutorial/Mobile/SettingUpAndroidTutorial.SettingUpAndroidTutorial"),						IsAvailableOnWindows || IsAvailableOnMac,	true,				false,					true),
@@ -153,6 +154,17 @@ void UpdatePlatformSDKStatus(FString InPlatformName, EPlatformSDKStatus InStatus
 	}
 }
 
+void UpdatePlatformDisplayName(FString InPlatformName, FText InDisplayName)
+{
+	for (const FPlatformInfo& PlatformInfo : PlatformInfoArray)
+	{
+		if (PlatformInfo.TargetPlatformName == FName(*InPlatformName))
+		{
+			const_cast<FPlatformInfo&>(PlatformInfo).DisplayName = InDisplayName;
+		}
+	}
+}
+
 FPlatformEnumerator EnumeratePlatformInfoArray()
 {
 	return FPlatformEnumerator(PlatformInfoArray, ARRAY_COUNT(PlatformInfoArray));
diff --git a/Engine/Source/Developer/DesktopPlatform/Public/PlatformInfo.h b/Engine/Source/Developer/DesktopPlatform/Public/PlatformInfo.h
--- a/Engine/Source/Developer/DesktopPlatform/Public/PlatformInfo.h
+++ b/Engine/Source/Developer/DesktopPlatform/Public/PlatformInfo.h
@@ -323,4 +323,10 @@ namespace PlatformInfo
 	 */
 	DESKTOPPLATFORM_API void UpdatePlatformSDKStatus(FString InPlatformName, EPlatformSDKStatus InStatus);
 
+	/**
+	* Update the display name for a platform
+	* @param InPlatformName - The platform to update
+	* @param InDisplayName - The new display name
+	*/
+	DESKTOPPLATFORM_API void UpdatePlatformDisplayName(FString InPlatformName, FText InDisplayName);
 }
diff --git a/Engine/Source/Developer/TargetPlatform/TargetPlatform.Build.cs b/Engine/Source/Developer/TargetPlatform/TargetPlatform.Build.cs
--- a/Engine/Source/Developer/TargetPlatform/TargetPlatform.Build.cs
+++ b/Engine/Source/Developer/TargetPlatform/TargetPlatform.Build.cs
@@ -61,7 +61,8 @@ public class TargetPlatform : ModuleRules
 					DynamicallyLoadedModuleNames.Add("Android_ETC1TargetPlatform");
 					DynamicallyLoadedModuleNames.Add("Android_ETC2TargetPlatform");
 					DynamicallyLoadedModuleNames.Add("Android_ASTCTargetPlatform");
-					DynamicallyLoadedModuleNames.Add("IOSTargetPlatform");
+                    DynamicallyLoadedModuleNames.Add("Android_MultiTargetPlatform");
+                    DynamicallyLoadedModuleNames.Add("IOSTargetPlatform");
                     DynamicallyLoadedModuleNames.Add("HTML5TargetPlatform");
 				}
 			}
@@ -85,6 +86,7 @@ public class TargetPlatform : ModuleRules
 				if (Target.Type == TargetRules.TargetType.Editor || Target.Type == TargetRules.TargetType.Program)
 				{
 					DynamicallyLoadedModuleNames.Add("AndroidTargetPlatform");
+					DynamicallyLoadedModuleNames.Add("Android_MultiTargetPlatform");
 					DynamicallyLoadedModuleNames.Add("Android_PVRTCTargetPlatform");
 					DynamicallyLoadedModuleNames.Add("Android_ATCTargetPlatform");
 					DynamicallyLoadedModuleNames.Add("Android_DXTTargetPlatform");
diff --git a/Engine/Source/Programs/UnrealBuildTool/Android/UEBuildAndroid.cs b/Engine/Source/Programs/UnrealBuildTool/Android/UEBuildAndroid.cs
--- a/Engine/Source/Programs/UnrealBuildTool/Android/UEBuildAndroid.cs
+++ b/Engine/Source/Programs/UnrealBuildTool/Android/UEBuildAndroid.cs
@@ -314,6 +314,7 @@ namespace UnrealBuildTool
 							Rules.PlatformSpecificDynamicallyLoadedModuleNames.Add("Android_ETC1TargetPlatform");
 							Rules.PlatformSpecificDynamicallyLoadedModuleNames.Add("Android_ETC2TargetPlatform");
 							Rules.PlatformSpecificDynamicallyLoadedModuleNames.Add("Android_ASTCTargetPlatform");
+                            Rules.PlatformSpecificDynamicallyLoadedModuleNames.Add("Android_MultiTargetPlatform");
                         }
 					}
 					else if (ModuleName == "TargetPlatform")
@@ -342,6 +343,7 @@ namespace UnrealBuildTool
 					    Rules.PlatformSpecificDynamicallyLoadedModuleNames.Add("Android_ETC1TargetPlatform");
 					    Rules.PlatformSpecificDynamicallyLoadedModuleNames.Add("Android_ETC2TargetPlatform");
 					    Rules.PlatformSpecificDynamicallyLoadedModuleNames.Add("Android_ASTCTargetPlatform");
+                        Rules.PlatformSpecificDynamicallyLoadedModuleNames.Add("Android_MultiTargetPlatform");
                     }
     
 				    if (bBuildShaderFormats)
diff --git a/Engine/Source/Runtime/Android/AndroidRuntimeSettings/AndroidRuntimeSettings.Build.cs b/Engine/Source/Runtime/Android/AndroidRuntimeSettings/AndroidRuntimeSettings.Build.cs
--- a/Engine/Source/Runtime/Android/AndroidRuntimeSettings/AndroidRuntimeSettings.Build.cs
+++ b/Engine/Source/Runtime/Android/AndroidRuntimeSettings/AndroidRuntimeSettings.Build.cs
@@ -15,5 +15,16 @@ public class AndroidRuntimeSettings : ModuleRules
 				"CoreUObject"
 			}
 		);
+
+        if (Target.Type == TargetRules.TargetType.Editor || Target.Type == TargetRules.TargetType.Program)
+        {
+            PrivateDependencyModuleNames.AddRange(
+                new string[]
+			    {
+                    "TargetPlatform",
+                    "Android_MultiTargetPlatform"
+			    }
+            );
+        }
 	}
 }
diff --git a/Engine/Source/Runtime/Android/AndroidRuntimeSettings/Classes/AndroidRuntimeSettings.h b/Engine/Source/Runtime/Android/AndroidRuntimeSettings/Classes/AndroidRuntimeSettings.h
--- a/Engine/Source/Runtime/Android/AndroidRuntimeSettings/Classes/AndroidRuntimeSettings.h
+++ b/Engine/Source/Runtime/Android/AndroidRuntimeSettings/Classes/AndroidRuntimeSettings.h
@@ -258,6 +258,30 @@ public:
 	UPROPERTY(GlobalConfig, EditAnywhere, Category = DataCooker, meta = (DisplayName = "Audio encoding"))
 	TEnumAsByte<EAndroidAudio::Type> AndroidAudio;
 	
+	/** Include ETC1 textures when packaging with the Android (Multi) variant. ETC1 will be included if no other formats are selected. */
+	UPROPERTY(GlobalConfig, EditAnywhere, Category = MultiTextureFormats, meta = (DisplayName = "Include ETC1 textures"))
+	bool bMultiTargetFormat_ETC1;
+
+	/** Include ETC2 textures when packaging with the Android (Multi) variant. */
+	UPROPERTY(GlobalConfig, EditAnywhere, Category = MultiTextureFormats, meta = (DisplayName = "Include ETC2 textures"))
+	bool bMultiTargetFormat_ETC2;
+
+	/** Include DXT textures when packaging with the Android (Multi) variant. */
+	UPROPERTY(GlobalConfig, EditAnywhere, Category = MultiTextureFormats, meta = (DisplayName = "Include DXT textures"))
+	bool bMultiTargetFormat_DXT;
+
+	/** Include PVRTC textures when packaging with the Android (Multi) variant. */
+	UPROPERTY(GlobalConfig, EditAnywhere, Category = MultiTextureFormats, meta = (DisplayName = "Include PVRTC textures"))
+	bool bMultiTargetFormat_PVRTC;
+
+	/** Include ATC textures when packaging with the Android (Multi) variant. */
+	UPROPERTY(GlobalConfig, EditAnywhere, Category = MultiTextureFormats, meta = (DisplayName = "Include ATC textures"))
+	bool bMultiTargetFormat_ATC;
+
+	/** Include ASTC textures when packaging with the Android (Multi) variant. */
+	UPROPERTY(GlobalConfig, EditAnywhere, Category = MultiTextureFormats, meta = (DisplayName = "Include ASTC textures"))
+	bool bMultiTargetFormat_ASTC;
+
 	
 #if WITH_EDITOR
 	// UObject interface
diff --git a/Engine/Source/Runtime/Android/AndroidRuntimeSettings/Private/AndroidRuntimeSettings.cpp b/Engine/Source/Runtime/Android/AndroidRuntimeSettings/Private/AndroidRuntimeSettings.cpp
--- a/Engine/Source/Runtime/Android/AndroidRuntimeSettings/Private/AndroidRuntimeSettings.cpp
+++ b/Engine/Source/Runtime/Android/AndroidRuntimeSettings/Private/AndroidRuntimeSettings.cpp
@@ -4,12 +4,22 @@
 
 #include "AndroidRuntimeSettings.h"
 
+#if WITH_EDITOR
+#include "TargetPlatform.h"
+#include "IAndroid_MultiTargetPlatformModule.h"
+#endif
+
 UAndroidRuntimeSettings::UAndroidRuntimeSettings(const FObjectInitializer& ObjectInitializer)
 	: Super(ObjectInitializer)
 	, Orientation(EAndroidScreenOrientation::Landscape)
 	, bEnableGooglePlaySupport(false)
+	, bMultiTargetFormat_ETC1(true)
+	, bMultiTargetFormat_ETC2(true)
+	, bMultiTargetFormat_DXT(true)
+	, bMultiTargetFormat_PVRTC(true)
+	, bMultiTargetFormat_ATC(true)
+	, bMultiTargetFormat_ASTC(true)
 {
-
 }
 
 #if WITH_EDITOR
@@ -30,6 +40,25 @@ void UAndroidRuntimeSettings::PostEditChangeProperty(struct FPropertyChangedEven
 		bBuildForES2 = true;
 		UpdateSinglePropertyInConfigFile(GetClass()->FindPropertyByName(GET_MEMBER_NAME_CHECKED(UAndroidRuntimeSettings, bBuildForES2)), GetDefaultConfigFilename());
 	}
+
+	if (PropertyChangedEvent.Property != nullptr && PropertyChangedEvent.Property->GetName().StartsWith(TEXT("bMultiTargetFormat")))
+	{
+		UpdateSinglePropertyInConfigFile(PropertyChangedEvent.Property, GetDefaultConfigFilename());
+
+		// Ensure we have at least one format for Android_Multi
+		if (!bMultiTargetFormat_ETC1 && !bMultiTargetFormat_ETC2 && !bMultiTargetFormat_DXT && !bMultiTargetFormat_PVRTC && !bMultiTargetFormat_ATC && !bMultiTargetFormat_ASTC)
+		{
+			bMultiTargetFormat_ETC1 = true;
+			UpdateSinglePropertyInConfigFile(GetClass()->FindPropertyByName(GET_MEMBER_NAME_CHECKED(UAndroidRuntimeSettings, bMultiTargetFormat_ETC1)), GetDefaultConfigFilename());
+		}
+
+		// Notify the Android_MultiTargetPlatform module if it's loaded
+		IAndroid_MultiTargetPlatformModule* Module = FModuleManager::GetModulePtr<IAndroid_MultiTargetPlatformModule>("Android_MultiTargetPlatform");
+		if (Module)
+		{
+			Module->NotifySelectedFormatsChanged();
+		}
+	}
 }
 
 void UAndroidRuntimeSettings::PostInitProperties()
diff --git a/Engine/Source/Runtime/Core/Public/Android/AndroidProperties.h b/Engine/Source/Runtime/Core/Public/Android/AndroidProperties.h
--- a/Engine/Source/Runtime/Core/Public/Android/AndroidProperties.h
+++ b/Engine/Source/Runtime/Core/Public/Android/AndroidProperties.h
@@ -184,3 +184,11 @@ struct FAndroid_ES31PlatformProperties : public FAndroidPlatformProperties
 		return "Android_ES31";
 	}
 };
+
+struct FAndroid_MultiPlatformProperties : public FAndroidPlatformProperties
+{
+	static FORCEINLINE const char* PlatformName()
+	{
+		return "Android_Multi";
+	}
+};
