#if WITH_ENGINE fix for CL 2698623

diff --git a/Engine/Source/Developer/Android/Android_MultiTargetPlatform/Private/Android_MultiTargetPlatformModule.cpp b/Engine/Source/Developer/Android/Android_MultiTargetPlatform/Private/Android_MultiTargetPlatformModule.cpp
--- a/Engine/Source/Developer/Android/Android_MultiTargetPlatform/Private/Android_MultiTargetPlatformModule.cpp
+++ b/Engine/Source/Developer/Android/Android_MultiTargetPlatform/Private/Android_MultiTargetPlatformModule.cpp
@@ -81,6 +81,7 @@ class FAndroid_MultiTargetPlatform
 		return FString(FAndroid_MultiPlatformProperties::PlatformName());
 	}
 
+#if WITH_ENGINE
 	virtual void GetTextureFormats(const UTexture* Texture, TArray<FName>& OutFormats) const
 	{
 		// Ask each platform variant to choose texture formats
@@ -94,7 +95,8 @@ class FAndroid_MultiTargetPlatform
 			}
 		}
 	}
-	
+#endif	
+
 	virtual FText GetVariantDisplayName() const override
 	{
 		return LOCTEXT("Android_Multi_ShortName", "Multi");
