Fix up duplicate global symbols that cause a monolithic editor to fail to compile
#rnx

diff --git a/Engine/Plugins/2D/Paper2D/Source/Paper2DEditor/Private/ContentBrowserExtensions/ContentBrowserExtensions.cpp b/Engine/Plugins/2D/Paper2D/Source/Paper2DEditor/Private/ContentBrowserExtensions/ContentBrowserExtensions.cpp
--- a/Engine/Plugins/2D/Paper2D/Source/Paper2DEditor/Private/ContentBrowserExtensions/ContentBrowserExtensions.cpp
+++ b/Engine/Plugins/2D/Paper2D/Source/Paper2DEditor/Private/ContentBrowserExtensions/ContentBrowserExtensions.cpp
@@ -28,8 +28,8 @@ DEFINE_LOG_CATEGORY(LogPaperCBExtensions);
 
 //////////////////////////////////////////////////////////////////////////
 
-FContentBrowserMenuExtender_SelectedAssets ContentBrowserExtenderDelegate;
-FDelegateHandle ContentBrowserExtenderDelegateHandle;
+static FContentBrowserMenuExtender_SelectedAssets ContentBrowserExtenderDelegate;
+static FDelegateHandle ContentBrowserExtenderDelegateHandle;
 
 //////////////////////////////////////////////////////////////////////////
 // FContentBrowserSelectedAssetExtensionBase
diff --git a/Engine/Source/Developer/Android/AndroidTargetPlatform/Private/AndroidTargetPlatform.cpp b/Engine/Source/Developer/Android/AndroidTargetPlatform/Private/AndroidTargetPlatform.cpp
--- a/Engine/Source/Developer/Android/AndroidTargetPlatform/Private/AndroidTargetPlatform.cpp
+++ b/Engine/Source/Developer/Android/AndroidTargetPlatform/Private/AndroidTargetPlatform.cpp
@@ -637,41 +637,44 @@ void FAndroidTargetPlatform::GetAllWaveFormats(TArray<FName>& OutFormats) const
 	OutFormats.Add(NAME_ADPCM);
 }
 
-void CachePlatformAudioCookOverrides(FPlatformAudioCookOverrides& OutOverrides)
+namespace
 {
-	const TCHAR* CategoryName = TEXT("/Script/AndroidRuntimeSettings.AndroidRuntimeSettings");
+	void CachePlatformAudioCookOverrides(FPlatformAudioCookOverrides& OutOverrides)
+	{
+		const TCHAR* CategoryName = TEXT("/Script/AndroidRuntimeSettings.AndroidRuntimeSettings");
 
-	GConfig->GetBool(CategoryName, TEXT("bResampleForDevice"), OutOverrides.bResampleForDevice, GEngineIni);
-	
-	GConfig->GetFloat(CategoryName, TEXT("CompressionQualityModifier"), OutOverrides.CompressionQualityModifier, GEngineIni);
+		GConfig->GetBool(CategoryName, TEXT("bResampleForDevice"), OutOverrides.bResampleForDevice, GEngineIni);
 
-	//Cache sample rate map.
-	OutOverrides.PlatformSampleRates.Reset();
+		GConfig->GetFloat(CategoryName, TEXT("CompressionQualityModifier"), OutOverrides.CompressionQualityModifier, GEngineIni);
 
-	float RetrievedSampleRate = -1.0f;
+		//Cache sample rate map.
+		OutOverrides.PlatformSampleRates.Reset();
 
-	GConfig->GetFloat(CategoryName, TEXT("MaxSampleRate"), RetrievedSampleRate, GEngineIni);
-	OutOverrides.PlatformSampleRates.Add(ESoundwaveSampleRateSettings::Max, RetrievedSampleRate);
+		float RetrievedSampleRate = -1.0f;
 
-	RetrievedSampleRate = -1.0f;
+		GConfig->GetFloat(CategoryName, TEXT("MaxSampleRate"), RetrievedSampleRate, GEngineIni);
+		OutOverrides.PlatformSampleRates.Add(ESoundwaveSampleRateSettings::Max, RetrievedSampleRate);
 
-	GConfig->GetFloat(CategoryName, TEXT("HighSampleRate"), RetrievedSampleRate, GEngineIni);
-	OutOverrides.PlatformSampleRates.Add(ESoundwaveSampleRateSettings::High, RetrievedSampleRate);
+		RetrievedSampleRate = -1.0f;
 
-	RetrievedSampleRate = -1.0f;
+		GConfig->GetFloat(CategoryName, TEXT("HighSampleRate"), RetrievedSampleRate, GEngineIni);
+		OutOverrides.PlatformSampleRates.Add(ESoundwaveSampleRateSettings::High, RetrievedSampleRate);
 
-	GConfig->GetFloat(CategoryName, TEXT("MedSampleRate"), RetrievedSampleRate, GEngineIni);
-	OutOverrides.PlatformSampleRates.Add(ESoundwaveSampleRateSettings::Medium, RetrievedSampleRate);
+		RetrievedSampleRate = -1.0f;
 
-	RetrievedSampleRate = -1.0f;
+		GConfig->GetFloat(CategoryName, TEXT("MedSampleRate"), RetrievedSampleRate, GEngineIni);
+		OutOverrides.PlatformSampleRates.Add(ESoundwaveSampleRateSettings::Medium, RetrievedSampleRate);
 
-	GConfig->GetFloat(CategoryName, TEXT("LowSampleRate"), RetrievedSampleRate, GEngineIni);
-	OutOverrides.PlatformSampleRates.Add(ESoundwaveSampleRateSettings::Low, RetrievedSampleRate);
+		RetrievedSampleRate = -1.0f;
 
-	RetrievedSampleRate = -1.0f;
+		GConfig->GetFloat(CategoryName, TEXT("LowSampleRate"), RetrievedSampleRate, GEngineIni);
+		OutOverrides.PlatformSampleRates.Add(ESoundwaveSampleRateSettings::Low, RetrievedSampleRate);
 
-	GConfig->GetFloat(CategoryName, TEXT("MinSampleRate"), RetrievedSampleRate, GEngineIni);
-	OutOverrides.PlatformSampleRates.Add(ESoundwaveSampleRateSettings::Min, RetrievedSampleRate);
+		RetrievedSampleRate = -1.0f;
+
+		GConfig->GetFloat(CategoryName, TEXT("MinSampleRate"), RetrievedSampleRate, GEngineIni);
+		OutOverrides.PlatformSampleRates.Add(ESoundwaveSampleRateSettings::Min, RetrievedSampleRate);
+	}
 }
 
 FPlatformAudioCookOverrides* FAndroidTargetPlatform::GetAudioCompressionSettings() const
diff --git a/Engine/Source/Developer/IOS/IOSTargetPlatform/Private/IOSTargetPlatform.cpp b/Engine/Source/Developer/IOS/IOSTargetPlatform/Private/IOSTargetPlatform.cpp
--- a/Engine/Source/Developer/IOS/IOSTargetPlatform/Private/IOSTargetPlatform.cpp
+++ b/Engine/Source/Developer/IOS/IOSTargetPlatform/Private/IOSTargetPlatform.cpp
@@ -623,41 +623,44 @@ void FIOSTargetPlatform::GetAllWaveFormats(TArray<FName>& OutFormat) const
 	OutFormat.Add(NAME_ADPCM);
 }
 
-void CachePlatformAudioCookOverrides(FPlatformAudioCookOverrides& OutOverrides)
+namespace
 {
-	const TCHAR* CategoryName = TEXT("/Script/IOSRuntimeSettings.IOSRuntimeSettings");
+	void CachePlatformAudioCookOverrides(FPlatformAudioCookOverrides& OutOverrides)
+	{
+		const TCHAR* CategoryName = TEXT("/Script/IOSRuntimeSettings.IOSRuntimeSettings");
 
-	GConfig->GetBool(CategoryName, TEXT("bResampleForDevice"), OutOverrides.bResampleForDevice, GEngineIni);
+		GConfig->GetBool(CategoryName, TEXT("bResampleForDevice"), OutOverrides.bResampleForDevice, GEngineIni);
 
-	GConfig->GetFloat(CategoryName, TEXT("CompressionQualityModifier"), OutOverrides.CompressionQualityModifier, GEngineIni);
+		GConfig->GetFloat(CategoryName, TEXT("CompressionQualityModifier"), OutOverrides.CompressionQualityModifier, GEngineIni);
 
-	//Cache sample rate map.
-	OutOverrides.PlatformSampleRates.Reset();
+		//Cache sample rate map.
+		OutOverrides.PlatformSampleRates.Reset();
 
-	float RetrievedSampleRate = -1.0f;
+		float RetrievedSampleRate = -1.0f;
 
-	GConfig->GetFloat(CategoryName, TEXT("MaxSampleRate"), RetrievedSampleRate, GEngineIni);
-	OutOverrides.PlatformSampleRates.Add(ESoundwaveSampleRateSettings::Max, RetrievedSampleRate);
+		GConfig->GetFloat(CategoryName, TEXT("MaxSampleRate"), RetrievedSampleRate, GEngineIni);
+		OutOverrides.PlatformSampleRates.Add(ESoundwaveSampleRateSettings::Max, RetrievedSampleRate);
 
-	RetrievedSampleRate = -1.0f;
+		RetrievedSampleRate = -1.0f;
 
-	GConfig->GetFloat(CategoryName, TEXT("HighSampleRate"), RetrievedSampleRate, GEngineIni);
-	OutOverrides.PlatformSampleRates.Add(ESoundwaveSampleRateSettings::High, RetrievedSampleRate);
+		GConfig->GetFloat(CategoryName, TEXT("HighSampleRate"), RetrievedSampleRate, GEngineIni);
+		OutOverrides.PlatformSampleRates.Add(ESoundwaveSampleRateSettings::High, RetrievedSampleRate);
 
-	RetrievedSampleRate = -1.0f;
+		RetrievedSampleRate = -1.0f;
 
-	GConfig->GetFloat(CategoryName, TEXT("MedSampleRate"), RetrievedSampleRate, GEngineIni);
-	OutOverrides.PlatformSampleRates.Add(ESoundwaveSampleRateSettings::Medium, RetrievedSampleRate);
+		GConfig->GetFloat(CategoryName, TEXT("MedSampleRate"), RetrievedSampleRate, GEngineIni);
+		OutOverrides.PlatformSampleRates.Add(ESoundwaveSampleRateSettings::Medium, RetrievedSampleRate);
 
-	RetrievedSampleRate = -1.0f;
+		RetrievedSampleRate = -1.0f;
 
-	GConfig->GetFloat(CategoryName, TEXT("LowSampleRate"), RetrievedSampleRate, GEngineIni);
-	OutOverrides.PlatformSampleRates.Add(ESoundwaveSampleRateSettings::Low, RetrievedSampleRate);
+		GConfig->GetFloat(CategoryName, TEXT("LowSampleRate"), RetrievedSampleRate, GEngineIni);
+		OutOverrides.PlatformSampleRates.Add(ESoundwaveSampleRateSettings::Low, RetrievedSampleRate);
 
-	RetrievedSampleRate = -1.0f;
+		RetrievedSampleRate = -1.0f;
 
-	GConfig->GetFloat(CategoryName, TEXT("MinSampleRate"), RetrievedSampleRate, GEngineIni);
-	OutOverrides.PlatformSampleRates.Add(ESoundwaveSampleRateSettings::Min, RetrievedSampleRate);
+		GConfig->GetFloat(CategoryName, TEXT("MinSampleRate"), RetrievedSampleRate, GEngineIni);
+		OutOverrides.PlatformSampleRates.Add(ESoundwaveSampleRateSettings::Min, RetrievedSampleRate);
+	}
 }
 
 FPlatformAudioCookOverrides* FIOSTargetPlatform::GetAudioCompressionSettings() const
diff --git a/Engine/Source/Editor/Blutility/Private/BlutilityContentBrowserExtensions.cpp b/Engine/Source/Editor/Blutility/Private/BlutilityContentBrowserExtensions.cpp
--- a/Engine/Source/Editor/Blutility/Private/BlutilityContentBrowserExtensions.cpp
+++ b/Engine/Source/Editor/Blutility/Private/BlutilityContentBrowserExtensions.cpp
@@ -24,8 +24,8 @@
 
 #define LOCTEXT_NAMESPACE "BlutilityContentBrowserExtensions"
 
-FContentBrowserMenuExtender_SelectedAssets ContentBrowserExtenderDelegate;
-FDelegateHandle ContentBrowserExtenderDelegateHandle;
+static FContentBrowserMenuExtender_SelectedAssets ContentBrowserExtenderDelegate;
+static FDelegateHandle ContentBrowserExtenderDelegateHandle;
 
 class FBlutilityContentBrowserExtensions_Impl
 {
