UE-57505 - Crash launching ES31 app on Android
The shader platform is now read from GConfig instead of using the (cached) out of date internal data.

diff --git a/Engine/Source/Developer/Android/AndroidTargetPlatform/Private/AndroidTargetPlatform.cpp b/Engine/Source/Developer/Android/AndroidTargetPlatform/Private/AndroidTargetPlatform.cpp
--- a/Engine/Source/Developer/Android/AndroidTargetPlatform/Private/AndroidTargetPlatform.cpp
+++ b/Engine/Source/Developer/Android/AndroidTargetPlatform/Private/AndroidTargetPlatform.cpp
@@ -213,7 +213,7 @@ bool FAndroidTargetPlatform::SupportsES2() const
 	// default to support ES2
 	bool bBuildForES2 = true;
 #if WITH_ENGINE
-	EngineSettings.GetBool(TEXT("/Script/AndroidRuntimeSettings.AndroidRuntimeSettings"), TEXT("bBuildForES2"), bBuildForES2);
+	GConfig->GetBool(TEXT("/Script/AndroidRuntimeSettings.AndroidRuntimeSettings"), TEXT("bBuildForES2"), bBuildForES2, GEngineIni);
 #endif
 	return bBuildForES2;
 }
@@ -223,7 +223,7 @@ bool FAndroidTargetPlatform::SupportsES31() const
 	// default no support for ES31
 	bool bBuildForES31 = false;
 #if WITH_ENGINE
-	EngineSettings.GetBool(TEXT("/Script/AndroidRuntimeSettings.AndroidRuntimeSettings"), TEXT("bBuildForES31"), bBuildForES31);
+	GConfig->GetBool(TEXT("/Script/AndroidRuntimeSettings.AndroidRuntimeSettings"), TEXT("bBuildForES31"), bBuildForES31, GEngineIni);
 #endif
 	return bBuildForES31;
 }
@@ -238,7 +238,7 @@ bool FAndroidTargetPlatform::SupportsVulkan() const
 	// default to not supporting Vulkan
 	bool bSupportsVulkan = false;
 #if WITH_ENGINE
-	EngineSettings.GetBool(TEXT("/Script/AndroidRuntimeSettings.AndroidRuntimeSettings"), TEXT("bSupportsVulkan"), bSupportsVulkan);
+	GConfig->GetBool(TEXT("/Script/AndroidRuntimeSettings.AndroidRuntimeSettings"), TEXT("bSupportsVulkan"), bSupportsVulkan, GEngineIni);
 #endif
 	return bSupportsVulkan;
 }
@@ -249,7 +249,7 @@ bool FAndroidTargetPlatform::SupportsSoftwareOcclusion() const
 	bool bSupportsSoftwareOcclusion = false;
 #if WITH_ENGINE
 	int32 IntValue = 0;
-	EngineSettings.GetInt(TEXT("ConsoleVariables"), TEXT("r.Mobile.AllowSoftwareOcclusion"), IntValue);
+	GConfig->GetInt(TEXT("ConsoleVariables"), TEXT("r.Mobile.AllowSoftwareOcclusion"), IntValue, GEngineIni);
 	bSupportsSoftwareOcclusion = (IntValue != 0);
 #endif
 	return bSupportsSoftwareOcclusion;
