Fixed: Enabling RenderDoc under Android Project Settings breaks RenderDoc 1.0
Removed Android RenderDoc packaging option, no longer needed

diff --git a/Engine/Source/Developer/Android/AndroidPlatformEditor/Private/AndroidTargetSettingsCustomization.cpp b/Engine/Source/Developer/Android/AndroidPlatformEditor/Private/AndroidTargetSettingsCustomization.cpp
--- a/Engine/Source/Developer/Android/AndroidPlatformEditor/Private/AndroidTargetSettingsCustomization.cpp
+++ b/Engine/Source/Developer/Android/AndroidPlatformEditor/Private/AndroidTargetSettingsCustomization.cpp
@@ -739,39 +739,6 @@ static FText GetAdrenoProfilerHelpText()
 	return FText::Format(LOCTEXT("AdrenoHelpText","{0}\n{1}"), Args);
 }
 
-static FText GetRenderDocHelpText()
-{
-	const static FText InstallText(LOCTEXT("RDOCInstallText", "Run the following command from a host command line from the android/apk/32 directory located in the installation directory of the RenderDoc tool, to install the RenderDocCmd application on your device."));
-	const static FString InstallCommand(TEXT("adb install -r RenderDocCmd.apk"));
-
-	const static FText RunText0(LOCTEXT("RDOCRunText0", "Open RenderDoc on the host"));
-	const static FText RunText1(LOCTEXT("RDOCRunText1", "1. In Tools -> Options -> Android, set the path to your adb executable."));
-	const static FText RunText2(LOCTEXT("RDOCRunText2", "2. Start the Remote Server using Tools -> Start Android Remote Server."));
-	const static FText RunText3(LOCTEXT("RDOCRunText3", "3. Check your device screen and 'Allow' RenderDocCmd to access files on your device."));
-	const static FText RunText4(LOCTEXT("RDOCRunText4", "4. Change your current Replay Context to your device using the bottom left menu, which should now show your device as Online."));
-	const static FText RunText5(LOCTEXT("RDOCRunText5", "5. In the capture executable tab, there is a button on the right of Executable Path that lets you select an installed Android package for capture."));
-	const static FText RunText6(LOCTEXT("RDOCRunText6", "6. Select your package and press the Launch button in the bottom right of the tab to start the package on the device."));
-	const static FText RunText7(LOCTEXT("RDOCRunText7", "7. If everything went well, a new tab will open with a button to Trigger captures."));
-
-	const static FText NoteText(LOCTEXT("RDOCNoteText", "If the latest RenderDoc release does not have Android functionality, download the latest nightly build."));
-
-	FFormatOrderedArguments Args;
-	Args.Add(InstallText);
-	Args.Add(FText::FromString(InstallCommand));
-	Args.Add(RunText0);
-	Args.Add(RunText1);
-	Args.Add(RunText2);
-	Args.Add(RunText3);
-	Args.Add(RunText4);
-	Args.Add(RunText5);
-	Args.Add(RunText6);
-	Args.Add(RunText7);
-	Args.Add(NoteText);
-
-	return FText::Format(LOCTEXT("RDOCHelpText","<RichTextBlock.TextHighlight>Installation</>\n{0}\n{1}\n\n<RichTextBlock.TextHighlight>Run</>\n{2}\n{3}\n{4}\n{5}\n{6}\n{7}\n{8}\n{9}\n\n<RichTextBlock.TextHighlight>Note</>\n{10}"), 
-		Args);
-}
-
 void FAndroidTargetSettingsCustomization::BuildGraphicsDebuggerSection(IDetailLayoutBuilder& DetailLayout)
 {
 	IDetailCategoryBuilder& GraphicsDebuggerCategory = DetailLayout.EditCategory(TEXT("GraphicsDebugger"));
@@ -865,52 +832,6 @@ void FAndroidTargetSettingsCustomization::BuildGraphicsDebuggerSection(IDetailLa
 			]
 		];
 	}
-
-	// RenderDoc settings
-	{
-		TAttribute<EVisibility> RenderDocSettingsVisibility(
-			TAttribute<EVisibility>::Create(TAttribute<EVisibility>::FGetter::CreateStatic(GraphicsDebuggerSettingsVisibility, EAndroidGraphicsDebugger::RenderDoc, AndroidGraphicsDebuggerProperty))
-		);
-
-		TSharedPtr<IPropertyHandle> RenderDocPathProperty = DetailLayout.GetProperty(GET_MEMBER_NAME_CHECKED(UAndroidRuntimeSettings, RenderDocPath));
-		DetailLayout.HideProperty(RenderDocPathProperty);
-		GraphicsDebuggerCategory.AddProperty(RenderDocPathProperty).Visibility(RenderDocSettingsVisibility);
-		
-		FText RenderDocHelpText = GetRenderDocHelpText();
-		
-		GraphicsDebuggerCategory.AddCustomRow(LOCTEXT("RenderDocInfo", "RenderDoc Info"), false)
-		.Visibility(RenderDocSettingsVisibility)
-		.WholeRowWidget
-		[
-			SNew(SBorder)
-			.Padding(1)
-			[
-				SNew(SVerticalBox)
-				+ SVerticalBox::Slot()
-				.Padding(FMargin(10, 10, 10, 10))
-				.AutoHeight()
-				[
-					SNew(SRichTextBlock)
-					.Text(RenderDocHelpText)
-					.TextStyle(FEditorStyle::Get(), "MessageLog")
-					.DecoratorStyleSet(&FEditorStyle::Get())
-					.AutoWrapText(true)
-				]
-				+ SVerticalBox::Slot()
-				.AutoHeight()
-				.Padding(FMargin(10, 10, 10, 10))
-				[	
-					SNew(SBox)
-					.HAlign(HAlign_Left)
-					[
-						SNew(SHyperlinkLaunchURL, TEXT("https://renderdoc.org/"))
-						.Text(LOCTEXT("RenderDocPage", "RenderDoc home page"))
-						.ToolTipText(LOCTEXT("RenderDocPageTooltip", "Opens the RenderDoc home page"))
-					]
-				]
-			]
-		];
-	}
 }
 
 //////////////////////////////////////////////////////////////////////////
diff --git a/Engine/Source/Programs/UnrealBuildTool/Platform/Android/UEBuildAndroid.cs b/Engine/Source/Programs/UnrealBuildTool/Platform/Android/UEBuildAndroid.cs
--- a/Engine/Source/Programs/UnrealBuildTool/Platform/Android/UEBuildAndroid.cs
+++ b/Engine/Source/Programs/UnrealBuildTool/Platform/Android/UEBuildAndroid.cs
@@ -445,8 +445,6 @@ namespace UnrealBuildTool
 				//LinkEnvironment.AdditionalLibraries.Add("Nvidia_gfx_debugger_stub");
 			}
 
-			SetupGraphicsDebugger(Target, CompileEnvironment, LinkEnvironment);
-
 			if (!UseTegraGraphicsDebugger(Target))
 			{
 				LinkEnvironment.AdditionalLibraries.Add("GLESv2");
@@ -492,28 +490,6 @@ namespace UnrealBuildTool
 			return false;
 		}
 
-		private bool SetupGraphicsDebugger(ReadOnlyTargetRules Target, CppCompileEnvironment CompileEnvironment, LinkEnvironment LinkEnvironment)
-		{
-			string AndroidGraphicsDebugger;
-			ConfigHierarchy Ini = ConfigCache.ReadHierarchy(ConfigHierarchyType.Engine, DirectoryReference.FromFile(Target.ProjectFile), UnrealTargetPlatform.Android);
-			Ini.GetString("/Script/AndroidRuntimeSettings.AndroidRuntimeSettings", "AndroidGraphicsDebugger", out AndroidGraphicsDebugger);
-
-			if (AndroidGraphicsDebugger.ToLower() == "renderdoc")
-			{
-				string RenderDocPath;
-				AndroidPlatformSDK.GetPath(Ini, "/Script/AndroidRuntimeSettings.AndroidRuntimeSettings", "RenderDocPath", out RenderDocPath);
-				string RenderDocLibPath = Path.Combine(RenderDocPath, @"android\lib\armeabi-v7a");
-				if (Directory.Exists(RenderDocLibPath))
-				{
-					LinkEnvironment.LibraryPaths.Add(new DirectoryReference(RenderDocLibPath));
-					LinkEnvironment.AdditionalLibraries.Add("VkLayer_GLES_RenderDoc");
-					return true;
-				}
-			}
-
-			return false;
-		}
-
 		public override bool ShouldCreateDebugInfo(ReadOnlyTargetRules Target)
 		{
 			switch (Target.Configuration)
diff --git a/Engine/Source/Programs/UnrealBuildTool/Platform/Android/UEDeployAndroid.cs b/Engine/Source/Programs/UnrealBuildTool/Platform/Android/UEDeployAndroid.cs
--- a/Engine/Source/Programs/UnrealBuildTool/Platform/Android/UEDeployAndroid.cs
+++ b/Engine/Source/Programs/UnrealBuildTool/Platform/Android/UEDeployAndroid.cs
@@ -1311,21 +1311,6 @@ namespace UnrealBuildTool
 						}
 					}
 					break;
-				case "renderdoc":
-					{
-						string RenderDocPath;
-						AndroidPlatformSDK.GetPath(Ini, "/Script/AndroidRuntimeSettings.AndroidRuntimeSettings", "RenderDocPath", out RenderDocPath);
-						if (Directory.Exists(RenderDocPath))
-						{
-							Directory.CreateDirectory(Path.Combine(UE4BuildPath, "libs", NDKArch));
-							string RenderDocLibSrcPath = Path.Combine(RenderDocPath, @"android\lib", NDKArch, "libVkLayer_GLES_RenderDoc.so");
-							string RenderDocLibDstPath = Path.Combine(UE4BuildPath, "libs", NDKArch, "libVkLayer_GLES_RenderDoc.so");
-
-							Log.TraceInformation("Copying {0} to {1}", RenderDocLibSrcPath, RenderDocLibDstPath);
-							File.Copy(RenderDocLibSrcPath, RenderDocLibDstPath, true);
-						}
-					}
-					break;
 
 				// @TODO: Add NVIDIA Gfx Debugger
 				/*
@@ -4081,16 +4066,6 @@ namespace UnrealBuildTool
 											"\t\t\tLog.debug(\"libMGD.so not loaded.\");\n" +
 											"\t\t}\n";
 					break;
-				case "renderdoc":
-					LoadLibraryDefaults += "\t\ttry\n" +
-											"\t\t{\n" +
-											"\t\t\tSystem.loadLibrary(\"VkLayer_GLES_RenderDoc\");\n" +
-											"\t\t}\n" +
-											"\t\tcatch (java.lang.UnsatisfiedLinkError e)\n" +
-											"\t\t{\n" +
-											"\t\t\tLog.debug(\"libVkLayer_GLES_RenderDoc.so not loaded.\");\n" +
-											"\t\t}\n";
-					break;
 			}
 			
 			Dictionary<string, string> Replacements = new Dictionary<string, string>{
diff --git a/Engine/Source/Runtime/Android/AndroidRuntimeSettings/Classes/AndroidRuntimeSettings.h b/Engine/Source/Runtime/Android/AndroidRuntimeSettings/Classes/AndroidRuntimeSettings.h
--- a/Engine/Source/Runtime/Android/AndroidRuntimeSettings/Classes/AndroidRuntimeSettings.h
+++ b/Engine/Source/Runtime/Android/AndroidRuntimeSettings/Classes/AndroidRuntimeSettings.h
@@ -166,7 +166,6 @@ namespace EAndroidGraphicsDebugger
 		None = 0 UMETA(DisplayName = "None"),
 		Mali = 1 UMETA(DisplayName = "Mali Graphics Debugger", ToolTip = "Configure for Mali Graphics Debugger."),
 		Adreno = 2 UMETA(DisplayName = "Adreno Profiler", ToolTip = "Configure for Adreno Profiler."),
-		RenderDoc = 3 UMETA(DisplayName = "RenderDoc (Experimental)", ToolTip = "Configure for RenderDoc.")
 	};
 }
 
@@ -488,10 +487,6 @@ public:
 	UPROPERTY(GlobalConfig, EditAnywhere, Category = GraphicsDebugger)
 	FDirectoryPath MaliGraphicsDebuggerPath;
 
-	/** The path to your RenderDoc installation (eg C:/Program Files/RenderDoc) */
-	UPROPERTY(GlobalConfig, EditAnywhere, Category = GraphicsDebugger)
-	FDirectoryPath RenderDocPath;
-
 	/** Include ETC1 textures when packaging with the Android (Multi) variant. ETC1 will be included if no other formats are selected. */
 	UPROPERTY(GlobalConfig, EditAnywhere, Category = MultiTextureFormats, meta = (DisplayName = "Include ETC1 textures"))
 	bool bMultiTargetFormat_ETC1;
