[UEMOB-455] Make editor tool to dump out a device .json file for a connected Android device

diff --git a/Engine/Source/Developer/Android/AndroidDeviceDetection/AndroidDeviceDetection.Build.cs b/Engine/Source/Developer/Android/AndroidDeviceDetection/AndroidDeviceDetection.Build.cs
--- a/Engine/Source/Developer/Android/AndroidDeviceDetection/AndroidDeviceDetection.Build.cs
+++ b/Engine/Source/Developer/Android/AndroidDeviceDetection/AndroidDeviceDetection.Build.cs
@@ -12,7 +12,10 @@ public class AndroidDeviceDetection : ModuleRules
 			new string[]
 			{
 				"Core",
-			}
+				"Json",
+                "JsonUtilities",
+                "PIEPreviewDeviceProfileSelector"
+            }
 		);
 
 		PrivateIncludePathModuleNames.AddRange(
@@ -29,7 +32,15 @@ public class AndroidDeviceDetection : ModuleRules
 			}
 		);
 
-		if (Target.bCompileAgainstEngine)
+        PrivateIncludePaths.AddRange(
+            new string[]
+            {
+                "Editor/PIEPreviewDeviceProfileSelector/Private"
+            }
+        );
+
+
+        if (Target.bCompileAgainstEngine)
 		{
 			PrivateDependencyModuleNames.Add("Engine");
 		}
diff --git a/Engine/Source/Developer/Android/AndroidDeviceDetection/Private/AndroidDeviceDetectionModule.cpp b/Engine/Source/Developer/Android/AndroidDeviceDetection/Private/AndroidDeviceDetectionModule.cpp
--- a/Engine/Source/Developer/Android/AndroidDeviceDetection/Private/AndroidDeviceDetectionModule.cpp
+++ b/Engine/Source/Developer/Android/AndroidDeviceDetection/Private/AndroidDeviceDetectionModule.cpp
@@ -25,6 +25,12 @@
 #include "Interfaces/IAndroidDeviceDetectionModule.h"
 #include "ITcpMessagingModule.h"
 
+#include "PIEPreviewDeviceSpecification.h"
+#include "JsonObjectConverter.h"
+#include "Dom/JsonObject.h"
+#include "Misc/FileHelper.h"
+#include "Misc/MessageDialog.h"
+
 #define LOCTEXT_NAMESPACE "FAndroidDeviceDetectionModule" 
 
 DEFINE_LOG_CATEGORY_STATIC(AndroidDeviceDetectionLog, Log, All);
@@ -128,6 +134,99 @@ private:
 		return true;
 	}
 
+	// searches for 'DPIString' and 
+	int32 ExtractDPI(const FString& SurfaceFlingerOutput, const FString& DPIString)
+	{
+		int32 FoundDpi = INDEX_NONE;
+
+		int32 DpiIndex = SurfaceFlingerOutput.Find(DPIString);
+		if (DpiIndex != INDEX_NONE)
+		{
+			int32 StartIndex = INDEX_NONE;
+			for (int32 i = DpiIndex; i < SurfaceFlingerOutput.Len(); ++i)
+			{
+				// if we somehow hit a line break character something went wrong and no digits were found on this line
+				// we don't want to search the SurfaceFlinger feed so exit now
+				if (FChar::IsLinebreak(SurfaceFlingerOutput[i]))
+				{
+					break;
+				}
+
+				// search for the first digit aka the beginning of the DPI value
+				if (StartIndex == INDEX_NONE && FChar::IsDigit(SurfaceFlingerOutput[i]))
+				{
+					StartIndex = i;
+				}
+				// if we hit some non-numeric character extract the number and exit
+				else if (StartIndex != INDEX_NONE && !FChar::IsDigit(SurfaceFlingerOutput[i]))
+				{
+					FString str = SurfaceFlingerOutput.Mid(StartIndex, i - StartIndex);
+					FoundDpi = FCString::Atoi(*str);
+					break;
+				}
+			}
+		}
+
+		return FoundDpi;
+	}
+	// retrieve the string between 'InOutStartIndex' and the start position of the next 'Token' substring
+	// the white spaces of the resulting string are trimmed out at both ends 
+	FString ExtractNextToken(int32& InOutStartIndex, const FString& SurfaceFlingerOutput, const FString& Token)
+	{
+		FString OutString;
+		int32 StartIndex = InOutStartIndex;
+
+		int32 EndIndex = SurfaceFlingerOutput.Find(Token, ESearchCase::IgnoreCase, ESearchDir::FromStart, StartIndex);
+
+		if (EndIndex != INDEX_NONE)
+		{
+			InOutStartIndex = EndIndex + 1;
+			// the index should point to the position before the token start
+			--EndIndex;
+
+			for (int32 i = StartIndex; i < EndIndex; ++i)
+			{
+				if (!FChar::IsWhitespace(SurfaceFlingerOutput[i]))
+				{
+					StartIndex = i;
+					break;
+				}
+			}
+
+			for (int32 i = EndIndex; i > StartIndex; --i)
+			{
+				if (!FChar::IsWhitespace(SurfaceFlingerOutput[i]))
+				{
+					EndIndex = i;
+					break;
+				}
+			}
+
+			OutString = SurfaceFlingerOutput.Mid(StartIndex, FMath::Max(0, EndIndex - StartIndex + 1));
+		}
+
+		return OutString;
+	}
+
+	void ExtractGPUInfo(FString& outGLVersion, FString& outGPUFamily, const FString& SurfaceFlingerOutput)
+	{
+		int32 FoundDpi = INDEX_NONE;
+
+		int32 LineIndex = SurfaceFlingerOutput.Find(TEXT("GLES:"));
+		if (LineIndex != INDEX_NONE)
+		{
+			int32 StartIndex = SurfaceFlingerOutput.Find(TEXT(":"), ESearchCase::IgnoreCase, ESearchDir::FromStart, LineIndex);
+			if (StartIndex != INDEX_NONE)
+			{
+				++StartIndex;
+
+				FString GPUVendorString = ExtractNextToken(StartIndex, SurfaceFlingerOutput, TEXT(","));
+				outGPUFamily = ExtractNextToken(StartIndex, SurfaceFlingerOutput, TEXT(","));
+				outGLVersion = ExtractNextToken(StartIndex, SurfaceFlingerOutput, TEXT("\n"));
+			}
+		}
+	}
+
 	void QueryConnectedDevices()
 	{
 		// grab the list of devices via adb
@@ -241,6 +340,62 @@ private:
 					{
 						continue;
 					}
+
+					// extract DPI information
+					int32 XDpi = ExtractDPI(NewDeviceInfo.GLESExtensions, TEXT("x-dpi"));
+					int32 YDpi = ExtractDPI(NewDeviceInfo.GLESExtensions, TEXT("y-dpi"));
+
+					if (XDpi != INDEX_NONE && YDpi != INDEX_NONE)
+					{
+						NewDeviceInfo.DeviceDPI = (XDpi + YDpi) / 2;
+					}
+
+					// extract OpenGL version and GPU family name
+					ExtractGPUInfo(NewDeviceInfo.OpenGLVersionString, NewDeviceInfo.GPUFamilyString, NewDeviceInfo.GLESExtensions);
+				}
+
+				// grab device brand
+				{
+					FString ExecCommand = FString::Printf(TEXT("-s %s %s ro.product.brand"), *NewDeviceInfo.SerialNumber, *GetPropCommand);
+
+					FString RoProductBrand;
+					ExecuteAdbCommand(*ExecCommand, &RoProductBrand, nullptr);
+					const TCHAR* Ptr = *RoProductBrand;
+					FParse::Line(&Ptr, NewDeviceInfo.DeviceBrand);
+				}
+
+				// grab screen resolution
+				{
+					FString ResolutionString;
+					const FString ExecCommand = FString::Printf(TEXT("-s %s shell wm size"), *NewDeviceInfo.SerialNumber);
+					if (ExecuteAdbCommand(*ExecCommand, &ResolutionString, nullptr))
+					{
+						bool bFoundResX = false;
+						int32 StartIndex = INDEX_NONE;
+						for (int32 Index = 0; Index < ResolutionString.Len(); ++Index)
+						{
+							if (StartIndex == INDEX_NONE && FChar::IsDigit(ResolutionString[Index]))
+							{
+								StartIndex = Index;
+							}
+							else if (StartIndex != INDEX_NONE && !FChar::IsDigit(ResolutionString[Index]))
+							{
+								FString str = ResolutionString.Mid(StartIndex, Index - StartIndex);
+
+								if (bFoundResX)
+								{
+									NewDeviceInfo.ResolutionY = FCString::Atoi(*str);
+									break;
+								}
+								else
+								{
+									NewDeviceInfo.ResolutionX = FCString::Atoi(*str);
+									bFoundResX = true;
+									StartIndex = INDEX_NONE;
+								}
+							}
+						}
+					}
 				}
 
 				// grab the GL ES version
@@ -494,6 +649,89 @@ public:
 		DetectionThreadRunnable->UpdateADBPath(ADBPath, GetPropCommand, bGetExtensionsViaSurfaceFlinger, bForLumin);
 	}
 
+	virtual void ExportDeviceProfile(const FString& OutPath, const FString& DeviceName) override
+	{
+		// instantiate an FPIEPreviewDeviceSpecifications instance and its values
+		FPIEPreviewDeviceSpecifications DeviceSpecs;
+
+		bool bOpenGL3x = false;
+		{
+			FScopeLock ExportLock(GetDeviceMapLock());
+
+			const FAndroidDeviceInfo* DeviceInfo = GetDeviceMap().Find(DeviceName);
+			if (DeviceInfo == nullptr)
+			{
+				FText TitleMessage = LOCTEXT("loc_ExportError_Title", "File export error.");
+				FMessageDialog::Open(EAppMsgType::Ok, EAppReturnType::Ok, LOCTEXT("loc_ExportError_Message", "Device disconnected!"), &TitleMessage);
+				return;
+			}
+
+			// generic values
+			DeviceSpecs.DevicePlatform = EPIEPreviewDeviceType::Android;
+			DeviceSpecs.ResolutionX = DeviceInfo->ResolutionX;
+			DeviceSpecs.ResolutionY = DeviceInfo->ResolutionY;
+			DeviceSpecs.ResolutionYImmersiveMode = 0;
+			DeviceSpecs.PPI = DeviceInfo->DeviceDPI;
+			DeviceSpecs.ScaleFactors = { 0.25f, 0.5f, 0.75f, 1.0f };
+
+			// Android specific values
+			DeviceSpecs.AndroidProperties.AndroidVersion = DeviceInfo->HumanAndroidVersion;
+			DeviceSpecs.AndroidProperties.DeviceModel = DeviceInfo->Model;
+			DeviceSpecs.AndroidProperties.DeviceMake = DeviceInfo->DeviceBrand;
+			DeviceSpecs.AndroidProperties.GLVersion = DeviceInfo->OpenGLVersionString;
+			DeviceSpecs.AndroidProperties.GPUFamily = DeviceInfo->GPUFamilyString;
+			DeviceSpecs.AndroidProperties.VulkanVersion = "0.0.0";
+			DeviceSpecs.AndroidProperties.UsingHoudini = false;
+			DeviceSpecs.AndroidProperties.VulkanAvailable = false;
+
+			// OpenGL ES 3.x
+			bOpenGL3x = DeviceInfo->OpenGLVersionString.Contains(TEXT("OpenGL ES 3"));
+			if (bOpenGL3x)
+			{
+				DeviceSpecs.AndroidProperties.GLES31RHIState.SupportsInstancing = true;
+				DeviceSpecs.AndroidProperties.GLES31RHIState.MaxTextureDimensions = 2048;
+				DeviceSpecs.AndroidProperties.GLES31RHIState.MaxShadowDepthBufferSizeX = 2048;
+				DeviceSpecs.AndroidProperties.GLES31RHIState.MaxShadowDepthBufferSizeY = 2048;
+				DeviceSpecs.AndroidProperties.GLES31RHIState.MaxCubeTextureDimensions = 2048;
+				DeviceSpecs.AndroidProperties.GLES31RHIState.SupportsRenderTargetFormat_PF_G8 = true;
+				DeviceSpecs.AndroidProperties.GLES31RHIState.SupportsRenderTargetFormat_PF_FloatRGBA = DeviceInfo->GLESExtensions.Contains(TEXT("GL_EXT_color_buffer_half_float"));
+				DeviceSpecs.AndroidProperties.GLES31RHIState.SupportsMultipleRenderTargets = true;
+			}
+
+			// OpenGL ES 2.0
+			{
+				DeviceSpecs.AndroidProperties.GLES2RHIState.SupportsInstancing = false;
+				DeviceSpecs.AndroidProperties.GLES2RHIState.MaxTextureDimensions = 1024;
+				DeviceSpecs.AndroidProperties.GLES2RHIState.MaxShadowDepthBufferSizeX = 1024;
+				DeviceSpecs.AndroidProperties.GLES2RHIState.MaxShadowDepthBufferSizeY = 1024;
+				DeviceSpecs.AndroidProperties.GLES2RHIState.MaxCubeTextureDimensions = 512;
+				DeviceSpecs.AndroidProperties.GLES2RHIState.SupportsRenderTargetFormat_PF_G8 = true;
+				DeviceSpecs.AndroidProperties.GLES2RHIState.SupportsRenderTargetFormat_PF_FloatRGBA = DeviceInfo->GLESExtensions.Contains(TEXT("GL_EXT_color_buffer_half_float"));
+				DeviceSpecs.AndroidProperties.GLES2RHIState.SupportsMultipleRenderTargets = false;
+			}
+		} // FScopeLock ExportLock released
+
+		// create a JSon object from the above structure
+		TSharedPtr<FJsonObject> JsonObject = FJsonObjectConverter::UStructToJsonObject<FPIEPreviewDeviceSpecifications>(DeviceSpecs);
+
+		// if device does not support OpenGL 3.x avoid exporting anything about it
+		if (!bOpenGL3x)
+		{
+			JsonObject->RemoveField("GLES31RHIState");
+		}
+
+		// remove IOS fields
+		JsonObject->RemoveField("IOSProperties");
+
+		// serialize the JSon object to string
+		FString OutputString;
+		TSharedRef<TJsonWriter<>> Writer = TJsonWriterFactory<>::Create(&OutputString);
+		FJsonSerializer::Serialize(JsonObject.ToSharedRef(), Writer);
+
+		// export file to disk
+		FFileHelper::SaveStringToFile(OutputString, *OutPath);
+	} // end of virtual void ExportDeviceProfile(...)
+
 private:
 
 	// path to the adb command (local)
diff --git a/Engine/Source/Developer/Android/AndroidDeviceDetection/Public/Interfaces/IAndroidDeviceDetection.h b/Engine/Source/Developer/Android/AndroidDeviceDetection/Public/Interfaces/IAndroidDeviceDetection.h
--- a/Engine/Source/Developer/Android/AndroidDeviceDetection/Public/Interfaces/IAndroidDeviceDetection.h
+++ b/Engine/Source/Developer/Android/AndroidDeviceDetection/Public/Interfaces/IAndroidDeviceDetection.h
@@ -40,6 +40,22 @@ struct FAndroidDeviceInfo
 	// TCP port number on our local host forwarded over adb to the device
 	uint16 HostMessageBusPort;
 
+	// Holds pixel per inch value.
+	int32 DeviceDPI = 0;
+
+	// Holds the display resolution for the device
+	int32 ResolutionX = 0;
+	int32 ResolutionY = 0;
+
+	// Holds the reported OpenGLES version.
+	FString OpenGLVersionString;
+
+	// Holds the GPU family name.
+	FString GPUFamilyString;
+
+	// Holds the name of the manufacturer
+	FString DeviceBrand;
+
 	FAndroidDeviceInfo()
 		: SDKVersion(INDEX_NONE)
 		, GLESVersion(INDEX_NONE)
@@ -61,6 +77,8 @@ public:
 	virtual FCriticalSection* GetDeviceMapLock() = 0;
 	virtual void UpdateADBPath() = 0;
 	virtual FString GetADBPath() = 0;
+
+	virtual void ExportDeviceProfile(const FString& OutPath, const FString& DeviceName) = 0;
 protected:
 
 	/**
diff --git a/Engine/Source/Editor/UnrealEd/Private/Kismet2/DebuggerCommands.cpp b/Engine/Source/Editor/UnrealEd/Private/Kismet2/DebuggerCommands.cpp
--- a/Engine/Source/Editor/UnrealEd/Private/Kismet2/DebuggerCommands.cpp
+++ b/Engine/Source/Editor/UnrealEd/Private/Kismet2/DebuggerCommands.cpp
@@ -58,6 +58,10 @@
 
 #include "InstalledPlatformInfo.h"
 #include "PIEPreviewDeviceProfileSelectorModule.h"
+#include "IDesktopPlatform.h"
+#include "DesktopPlatformModule.h"
+#include "IAndroidDeviceDetectionModule.h"
+#include "IAndroidDeviceDetection.h"
 
 
 #define LOCTEXT_NAMESPACE "DebuggerCommands"
@@ -676,6 +680,76 @@ void FPlayWorldCommands::BuildToolbar( FToolBarBuilder& ToolbarBuilder, bool bIn
 	ToolbarBuilder.AddToolBarButton(FPlayWorldCommands::Get().StepOut, NAME_None, TAttribute<FText>(), TAttribute<FText>(), TAttribute<FSlateIcon>(), FName(TEXT("StepOut")));
 }
 
+// function will enumerate available Android devices that can export their profile to a json file
+// called (below) from AddAndroidConfigExportMenu()
+static void AddAndroidConfigExportSubMenus(FMenuBuilder& InMenuBuilder)
+{
+	IAndroidDeviceDetection* DeviceDetection = FModuleManager::LoadModuleChecked<IAndroidDeviceDetectionModule>("AndroidDeviceDetection").GetAndroidDeviceDetection();
+
+	TMap<FString, FAndroidDeviceInfo> AndroidDeviceMap;
+
+	// lock device map and copy its contents
+	{
+		FCriticalSection* DeviceLock = DeviceDetection->GetDeviceMapLock();
+		FScopeLock Lock(DeviceLock);
+		AndroidDeviceMap = DeviceDetection->GetDeviceMap();
+	}
+
+	for (auto& Pair : AndroidDeviceMap)
+	{
+		FAndroidDeviceInfo& DeviceInfo = Pair.Value;
+
+		FString ModelName = DeviceInfo.Model + TEXT("(") + DeviceInfo.DeviceBrand + TEXT(")");
+
+		// lambda function called to open the save dialog and trigger device export
+		auto LambdaSaveConfigFile = [DeviceName = Pair.Key, DefaultFileName = ModelName]()
+		{
+			TArray<FString> OutputFileName;
+			FString DefaultFolder = FPaths::EngineContentDir() + TEXT("Editor/PIEPreviewDeviceSpecs/Android/");
+
+			bool bResult = FDesktopPlatformModule::Get()->SaveFileDialog(
+				FSlateApplication::Get().FindBestParentWindowHandleForDialogs(nullptr),
+				LOCTEXT("PackagePluginDialogTitle", "Save platform configuration...").ToString(),
+				DefaultFolder,
+				DefaultFileName,
+				TEXT("Json config file (*.json)|*.json"),
+				0,
+				OutputFileName);
+
+			if (bResult && OutputFileName.Num())
+			{
+				IAndroidDeviceDetection* DeviceDetection = FModuleManager::LoadModuleChecked<IAndroidDeviceDetectionModule>("AndroidDeviceDetection").GetAndroidDeviceDetection();
+
+				if (DeviceDetection != nullptr)
+				{
+					DeviceDetection->ExportDeviceProfile(OutputFileName[0], DeviceName);
+				}
+			}
+		};
+
+		InMenuBuilder.AddMenuEntry(
+			FText::FromString(ModelName),
+			FText(),
+			FSlateIcon(FEditorStyle::GetStyleSetName(), "AssetEditor.SaveAsset"),
+			FUIAction(FExecuteAction::CreateLambda(LambdaSaveConfigFile))
+		);
+	}
+}
+
+// function adds a sub-menu that will enumerate Android devices whose profiles can be exported json files
+static void AddAndroidConfigExportMenu(FMenuBuilder& InMenuBuilder)
+{
+	InMenuBuilder.AddMenuSeparator();
+
+	InMenuBuilder.AddSubMenu(
+		LOCTEXT("loc_AddAndroidConfigExportMenu", "Export device settings"),
+		LOCTEXT("loc_tip_AddAndroidConfigExportMenu", "Export device settings to a Json file."),
+		FNewMenuDelegate::CreateStatic(&AddAndroidConfigExportSubMenus),
+		false,
+		FSlateIcon(FEditorStyle::GetStyleSetName(), "MainFrame.SaveAll")
+	);
+}
+
 static void MakePreviewDeviceMenu(FMenuBuilder& MenuBuilder )
 {
 	struct FLocal
@@ -690,22 +764,32 @@ static void MakePreviewDeviceMenu(FMenuBuilder& MenuBuilder )
 				MenuBuilderIn.AddMenuEntry(TargetedMobilePreviewDeviceCommands[Device]);
 			}
 
-			FText AndroidCategory = FText::FromString(TEXT("Android"));
-			FText IOSCategory = FText::FromString(TEXT("IOS"));
+			static FText AndroidCategory = FText::FromString(TEXT("Android"));
+			static FText IOSCategory = FText::FromString(TEXT("IOS"));
 
-			for (TSharedPtr<FPIEPreviewDeviceContainerCategory> SubCategory : PreviewDeviceCategory->GetSubCategories())
+			// Android devices can export their profile to a json file which then can be used for PIE device simulations
+			const FText& CategoryDisplayName = PreviewDeviceCategory->GetCategoryDisplayName();
+			if (CategoryDisplayName.CompareToCaseIgnored(AndroidCategory) == 0)
 			{
-				FText CategoryDisplayName = SubCategory->GetCategoryDisplayName();
-				
-				if (CategoryDisplayName.CompareToCaseIgnored(AndroidCategory) == 0) 
+				// check to see if we have any connected devices
+				bool bHasAndroidDevices = false;
 				{
-					CategoryDisplayName = FText(LOCTEXT("Android", "Android"));
+					IAndroidDeviceDetection* DeviceDetection = FModuleManager::LoadModuleChecked<IAndroidDeviceDetectionModule>("AndroidDeviceDetection").GetAndroidDeviceDetection();
+					FCriticalSection* DeviceLock = DeviceDetection->GetDeviceMapLock();
+
+					FScopeLock Lock(DeviceLock);
+					bHasAndroidDevices = DeviceDetection->GetDeviceMap().Num() > 0;
 				}
-				else if (CategoryDisplayName.CompareToCaseIgnored(IOSCategory) == 0)
+
+				// add the config. export menu
+				if (bHasAndroidDevices)
 				{
-					CategoryDisplayName = FText(LOCTEXT("IOS", "iOS"));
+					AddAndroidConfigExportMenu(MenuBuilderIn);
 				}
+			}
 
+			for (TSharedPtr<FPIEPreviewDeviceContainerCategory> SubCategory : PreviewDeviceCategory->GetSubCategories())
+			{
 				MenuBuilderIn.AddSubMenu(
 					SubCategory->GetCategoryDisplayName(),
 					SubCategory->GetCategoryToolTip(),
@@ -1067,10 +1151,10 @@ TSharedRef< SWidget > FPlayWorldCommands::GenerateLaunchMenuContent( TSharedRef<
 
 						// ... and add a menu entry
 						MenuBuilder.AddMenuEntry(
-							LaunchDeviceAction, 
+							LaunchDeviceAction,
 							ProjectTargetPlatformEditorModule.MakePlatformMenuItemWidget(*VanillaPlatform.PlatformInfo, true, Label),
-							NAME_None, 
-							Tooltip, 
+							NAME_None,
+							Tooltip,
 							EUserInterfaceActionType::Check
 							);
 					}
diff --git a/Engine/Source/Editor/UnrealEd/UnrealEd.Build.cs b/Engine/Source/Editor/UnrealEd/UnrealEd.Build.cs
--- a/Engine/Source/Editor/UnrealEd/UnrealEd.Build.cs
+++ b/Engine/Source/Editor/UnrealEd/UnrealEd.Build.cs
@@ -97,6 +97,7 @@ public class UnrealEd : ModuleRules
 				"NavigationSystem",
                 "MeshDescription",
                 "MeshBuilder",
+                "MaterialShaderQualitySettings"
             }
 		);
 
@@ -164,7 +165,7 @@ public class UnrealEd : ModuleRules
 				"PIEPreviewDeviceProfileSelector",
 				"PakFileUtilities",
 				"TimeManagement",
-			}
+            }
 		);
 
 		DynamicallyLoadedModuleNames.AddRange(
@@ -260,7 +261,9 @@ public class UnrealEd : ModuleRules
 		// Add include directory for Lightmass
 		PublicIncludePaths.Add("Programs/UnrealLightmass/Public");
 
-		PublicIncludePathModuleNames.AddRange(
+        PublicIncludePaths.Add("Developer/Android/AndroidDeviceDetection/Public/Interfaces");
+
+        PublicIncludePathModuleNames.AddRange(
 			new string[] {
 				"CollectionManager",
 				"BlueprintGraph",
@@ -271,10 +274,11 @@ public class UnrealEd : ModuleRules
 				"NavigationSystem",
 				"GameplayTasks",
 				"AIModule",
-			}
+            }
 			);
 
-		if ((Target.Platform == UnrealTargetPlatform.Win64) ||
+
+        if ((Target.Platform == UnrealTargetPlatform.Win64) ||
 			(Target.Platform == UnrealTargetPlatform.Win32))
 		{
 			PublicDependencyModuleNames.Add("XAudio2");
