Fix for ensure on cook for Switch/iOS/Android.

diff --git a/Engine/Source/Developer/Android/AndroidTargetPlatform/Private/AndroidTargetPlatform.cpp b/Engine/Source/Developer/Android/AndroidTargetPlatform/Private/AndroidTargetPlatform.cpp
--- a/Engine/Source/Developer/Android/AndroidTargetPlatform/Private/AndroidTargetPlatform.cpp
+++ b/Engine/Source/Developer/Android/AndroidTargetPlatform/Private/AndroidTargetPlatform.cpp
@@ -634,33 +634,94 @@ namespace
 
 		GConfig->GetFloat(CategoryName, TEXT("AutoStreamingThreshold"), OutOverrides.AutoStreamingThreshold, GEngineIni);
 
-		//Cache sample rate map.
-		OutOverrides.PlatformSampleRates.Reset();
-
+		//Cache sample rate map:
 		float RetrievedSampleRate = -1.0f;
 
 		GConfig->GetFloat(CategoryName, TEXT("MaxSampleRate"), RetrievedSampleRate, GEngineIni);
-		OutOverrides.PlatformSampleRates.Add(ESoundwaveSampleRateSettings::Max, RetrievedSampleRate);
+		float* FoundSampleRate = OutOverrides.PlatformSampleRates.Find(ESoundwaveSampleRateSettings::Max);
+
+		if (FoundSampleRate)
+		{
+			if (!FMath::IsNearlyEqual(*FoundSampleRate, RetrievedSampleRate))
+			{
+				*FoundSampleRate = RetrievedSampleRate;
+			}
+			
+		}
+		else
+		{
+			OutOverrides.PlatformSampleRates.Add(ESoundwaveSampleRateSettings::Max, RetrievedSampleRate);
+		}
 
 		RetrievedSampleRate = -1.0f;
 
 		GConfig->GetFloat(CategoryName, TEXT("HighSampleRate"), RetrievedSampleRate, GEngineIni);
-		OutOverrides.PlatformSampleRates.Add(ESoundwaveSampleRateSettings::High, RetrievedSampleRate);
+		FoundSampleRate = OutOverrides.PlatformSampleRates.Find(ESoundwaveSampleRateSettings::High);
+
+		if (FoundSampleRate)
+		{
+			if (!FMath::IsNearlyEqual(*FoundSampleRate, RetrievedSampleRate))
+			{
+				*FoundSampleRate = RetrievedSampleRate;
+			}
+
+		}
+		else
+		{
+			OutOverrides.PlatformSampleRates.Add(ESoundwaveSampleRateSettings::High, RetrievedSampleRate);
+		}
+
 
 		RetrievedSampleRate = -1.0f;
 
 		GConfig->GetFloat(CategoryName, TEXT("MedSampleRate"), RetrievedSampleRate, GEngineIni);
-		OutOverrides.PlatformSampleRates.Add(ESoundwaveSampleRateSettings::Medium, RetrievedSampleRate);
+		FoundSampleRate = OutOverrides.PlatformSampleRates.Find(ESoundwaveSampleRateSettings::Medium);
+
+		if (FoundSampleRate)
+		{
+			if (!FMath::IsNearlyEqual(*FoundSampleRate, RetrievedSampleRate))
+			{
+				*FoundSampleRate = RetrievedSampleRate;
+			}
+		}
+		else
+		{
+			OutOverrides.PlatformSampleRates.Add(ESoundwaveSampleRateSettings::Medium, RetrievedSampleRate);
+		}
 
 		RetrievedSampleRate = -1.0f;
 
 		GConfig->GetFloat(CategoryName, TEXT("LowSampleRate"), RetrievedSampleRate, GEngineIni);
-		OutOverrides.PlatformSampleRates.Add(ESoundwaveSampleRateSettings::Low, RetrievedSampleRate);
+		FoundSampleRate = OutOverrides.PlatformSampleRates.Find(ESoundwaveSampleRateSettings::Low);
+
+		if (FoundSampleRate)
+		{
+			if (!FMath::IsNearlyEqual(*FoundSampleRate, RetrievedSampleRate))
+			{
+				*FoundSampleRate = RetrievedSampleRate;
+			}
+		}
+		else
+		{
+			OutOverrides.PlatformSampleRates.Add(ESoundwaveSampleRateSettings::Low, RetrievedSampleRate);
+		}
 
 		RetrievedSampleRate = -1.0f;
 
 		GConfig->GetFloat(CategoryName, TEXT("MinSampleRate"), RetrievedSampleRate, GEngineIni);
-		OutOverrides.PlatformSampleRates.Add(ESoundwaveSampleRateSettings::Min, RetrievedSampleRate);
+		FoundSampleRate = OutOverrides.PlatformSampleRates.Find(ESoundwaveSampleRateSettings::Min);
+
+		if (FoundSampleRate)
+		{
+			if (!FMath::IsNearlyEqual(*FoundSampleRate, RetrievedSampleRate))
+			{
+				*FoundSampleRate = RetrievedSampleRate;
+			}
+		}
+		else
+		{
+			OutOverrides.PlatformSampleRates.Add(ESoundwaveSampleRateSettings::Min, RetrievedSampleRate);
+		}
 	}
 }
 
diff --git a/Engine/Source/Developer/IOS/IOSTargetPlatform/Private/IOSTargetPlatform.cpp b/Engine/Source/Developer/IOS/IOSTargetPlatform/Private/IOSTargetPlatform.cpp
--- a/Engine/Source/Developer/IOS/IOSTargetPlatform/Private/IOSTargetPlatform.cpp
+++ b/Engine/Source/Developer/IOS/IOSTargetPlatform/Private/IOSTargetPlatform.cpp
@@ -721,33 +721,94 @@ namespace
 
 		GConfig->GetFloat(CategoryName, TEXT("AutoStreamingThreshold"), OutOverrides.AutoStreamingThreshold, GEngineIni);
 
-		//Cache sample rate map.
-		OutOverrides.PlatformSampleRates.Reset();
-
+		//Cache sample rate map:
 		float RetrievedSampleRate = -1.0f;
 
 		GConfig->GetFloat(CategoryName, TEXT("MaxSampleRate"), RetrievedSampleRate, GEngineIni);
-		OutOverrides.PlatformSampleRates.Add(ESoundwaveSampleRateSettings::Max, RetrievedSampleRate);
+		float* FoundSampleRate = OutOverrides.PlatformSampleRates.Find(ESoundwaveSampleRateSettings::Max);
+
+		if (FoundSampleRate)
+		{
+			if (!FMath::IsNearlyEqual(*FoundSampleRate, RetrievedSampleRate))
+			{
+				*FoundSampleRate = RetrievedSampleRate;
+			}
+
+		}
+		else
+		{
+			OutOverrides.PlatformSampleRates.Add(ESoundwaveSampleRateSettings::Max, RetrievedSampleRate);
+		}
 
 		RetrievedSampleRate = -1.0f;
 
 		GConfig->GetFloat(CategoryName, TEXT("HighSampleRate"), RetrievedSampleRate, GEngineIni);
-		OutOverrides.PlatformSampleRates.Add(ESoundwaveSampleRateSettings::High, RetrievedSampleRate);
+		FoundSampleRate = OutOverrides.PlatformSampleRates.Find(ESoundwaveSampleRateSettings::High);
+
+		if (FoundSampleRate)
+		{
+			if (!FMath::IsNearlyEqual(*FoundSampleRate, RetrievedSampleRate))
+			{
+				*FoundSampleRate = RetrievedSampleRate;
+			}
+
+		}
+		else
+		{
+			OutOverrides.PlatformSampleRates.Add(ESoundwaveSampleRateSettings::High, RetrievedSampleRate);
+		}
+
 
 		RetrievedSampleRate = -1.0f;
 
 		GConfig->GetFloat(CategoryName, TEXT("MedSampleRate"), RetrievedSampleRate, GEngineIni);
-		OutOverrides.PlatformSampleRates.Add(ESoundwaveSampleRateSettings::Medium, RetrievedSampleRate);
+		FoundSampleRate = OutOverrides.PlatformSampleRates.Find(ESoundwaveSampleRateSettings::Medium);
+
+		if (FoundSampleRate)
+		{
+			if (!FMath::IsNearlyEqual(*FoundSampleRate, RetrievedSampleRate))
+			{
+				*FoundSampleRate = RetrievedSampleRate;
+			}
+		}
+		else
+		{
+			OutOverrides.PlatformSampleRates.Add(ESoundwaveSampleRateSettings::Medium, RetrievedSampleRate);
+		}
 
 		RetrievedSampleRate = -1.0f;
 
 		GConfig->GetFloat(CategoryName, TEXT("LowSampleRate"), RetrievedSampleRate, GEngineIni);
-		OutOverrides.PlatformSampleRates.Add(ESoundwaveSampleRateSettings::Low, RetrievedSampleRate);
+		FoundSampleRate = OutOverrides.PlatformSampleRates.Find(ESoundwaveSampleRateSettings::Low);
+
+		if (FoundSampleRate)
+		{
+			if (!FMath::IsNearlyEqual(*FoundSampleRate, RetrievedSampleRate))
+			{
+				*FoundSampleRate = RetrievedSampleRate;
+			}
+		}
+		else
+		{
+			OutOverrides.PlatformSampleRates.Add(ESoundwaveSampleRateSettings::Low, RetrievedSampleRate);
+		}
 
 		RetrievedSampleRate = -1.0f;
 
 		GConfig->GetFloat(CategoryName, TEXT("MinSampleRate"), RetrievedSampleRate, GEngineIni);
-		OutOverrides.PlatformSampleRates.Add(ESoundwaveSampleRateSettings::Min, RetrievedSampleRate);
+		FoundSampleRate = OutOverrides.PlatformSampleRates.Find(ESoundwaveSampleRateSettings::Min);
+
+		if (FoundSampleRate)
+		{
+			if (!FMath::IsNearlyEqual(*FoundSampleRate, RetrievedSampleRate))
+			{
+				*FoundSampleRate = RetrievedSampleRate;
+			}
+		}
+		else
+		{
+			OutOverrides.PlatformSampleRates.Add(ESoundwaveSampleRateSettings::Min, RetrievedSampleRate);
+		}
 	}
 }
 
diff --git a/Engine/Source/Runtime/AudioPlatformConfiguration/Public/AudioCompressionSettings.h b/Engine/Source/Runtime/AudioPlatformConfiguration/Public/AudioCompressionSettings.h
--- a/Engine/Source/Runtime/AudioPlatformConfiguration/Public/AudioCompressionSettings.h
+++ b/Engine/Source/Runtime/AudioPlatformConfiguration/Public/AudioCompressionSettings.h
@@ -63,13 +63,14 @@ struct FPlatformAudioCookOverrides
 		int32 ResampleBoolHash = (int32)InOverrides->bResampleForDevice;
 		OutSuffix.AppendInt(ResampleBoolHash);
 
-		for (auto& SampleRateQuality : InOverrides->PlatformSampleRates)
+		TMap<ESoundwaveSampleRateSettings, float> SampleRateMap = InOverrides->PlatformSampleRates;
+
+		for (auto& SampleRateQuality : SampleRateMap)
 		{
 			int32 SampleRateHash = FMath::FloorToInt(SampleRateQuality.Value / 1000.0f);
 			OutSuffix.AppendInt(SampleRateHash);
 		}
 	}
-
 };
 
 USTRUCT()
