Change the Feature Rendering Level back to SM5 when disabling Android ES2, ES3.1, or Vulkan if one is the active Feature Rendering Level

diff --git a/Engine/Source/Developer/Android/AndroidPlatformEditor/AndroidPlatformEditor.Build.cs b/Engine/Source/Developer/Android/AndroidPlatformEditor/AndroidPlatformEditor.Build.cs
--- a/Engine/Source/Developer/Android/AndroidPlatformEditor/AndroidPlatformEditor.Build.cs
+++ b/Engine/Source/Developer/Android/AndroidPlatformEditor/AndroidPlatformEditor.Build.cs
@@ -26,6 +26,7 @@ public class AndroidPlatformEditor : ModuleRules
                 "AndroidDeviceDetection",
                 "TargetPlatform",
                 "RenderCore",
+                "RHI",
                 "MaterialShaderQualitySettings",
 				"MainFrame",
                 "AudioSettingsEditor"
diff --git a/Engine/Source/Developer/Android/AndroidPlatformEditor/Private/AndroidPlatformEditorModule.cpp b/Engine/Source/Developer/Android/AndroidPlatformEditor/Private/AndroidPlatformEditorModule.cpp
--- a/Engine/Source/Developer/Android/AndroidPlatformEditor/Private/AndroidPlatformEditorModule.cpp
+++ b/Engine/Source/Developer/Android/AndroidPlatformEditor/Private/AndroidPlatformEditorModule.cpp
@@ -18,6 +18,7 @@
 #include "MaterialShaderQualitySettings.h"
 #include "ComponentRecreateRenderStateContext.h"
 #include "ShaderPlatformQualitySettings.h"
+#include "Misc/CoreDelegates.h"
 
 #define LOCTEXT_NAMESPACE "FAndroidPlatformEditorModule"
 
@@ -106,6 +107,40 @@ class FAndroidPlatformEditorModule
 
 		}
 
+
+		auto AndroidRuntimeSettings = GetMutableDefault<UAndroidRuntimeSettings>();
+		if (AndroidRuntimeSettings)
+		{
+			AndroidRuntimeSettings->OnPropertyChanged.AddLambda([this, AndroidRuntimeSettings](struct FPropertyChangedEvent& PropertyChangedEvent)
+				{
+					UEditorEngine* Editor = (UEditorEngine*)GEngine;
+					if (PropertyChangedEvent.Property != nullptr)
+					{
+						if (PropertyChangedEvent.Property->GetFName() == GET_MEMBER_NAME_CHECKED(UAndroidRuntimeSettings, bBuildForES31) &&
+							AndroidRuntimeSettings->bBuildForES31 == false)
+						{
+							FCoreDelegates::OnFeatureLevelDisabled.Broadcast(ERHIFeatureLevel::ES3_1, FName());
+							return;
+						}
+
+						if (PropertyChangedEvent.Property->GetFName() == GET_MEMBER_NAME_CHECKED(UAndroidRuntimeSettings, bBuildForES2) &&
+							AndroidRuntimeSettings->bBuildForES2 == false)
+						{
+							FCoreDelegates::OnFeatureLevelDisabled.Broadcast(ERHIFeatureLevel::ES2, FName());
+							return;
+						}
+
+						if (PropertyChangedEvent.Property->GetFName() == GET_MEMBER_NAME_CHECKED(UAndroidRuntimeSettings, bSupportsVulkan) &&
+							AndroidRuntimeSettings->bSupportsVulkan == false)
+						{
+							FCoreDelegates::OnFeatureLevelDisabled.Broadcast(ERHIFeatureLevel::ES3_1, LegacyShaderPlatformToShaderFormat(SP_VULKAN_ES3_1_ANDROID));
+							return;
+						}
+					}
+
+				});
+		}
+
 		// Force the SDK settings into a sane state initially so we can make use of them
 		auto &TargetPlatformManagerModule = FModuleManager::LoadModuleChecked<ITargetPlatformManagerModule>("TargetPlatform");
 		UAndroidSDKSettings * settings = GetMutableDefault<UAndroidSDKSettings>();
diff --git a/Engine/Source/Editor/UnrealEd/Private/EditorEngine.cpp b/Engine/Source/Editor/UnrealEd/Private/EditorEngine.cpp
--- a/Engine/Source/Editor/UnrealEd/Private/EditorEngine.cpp
+++ b/Engine/Source/Editor/UnrealEd/Private/EditorEngine.cpp
@@ -348,6 +348,21 @@ UEditorEngine::UEditorEngine(const FObjectInitializer& ObjectInitializer)
 	DefaultWorldFeatureLevel = GMaxRHIFeatureLevel;
 	PreviewFeatureLevel = DefaultWorldFeatureLevel;
 
+	FCoreDelegates::OnFeatureLevelDisabled.AddLambda([this](int RHIType, const FName& PreviewPlatformName)
+		{
+			ERHIFeatureLevel::Type FeatureLevelTypeToDisable = (ERHIFeatureLevel::Type)RHIType;
+			if (PreviewFeatureLevel == FeatureLevelTypeToDisable)
+			{
+				UMaterialShaderQualitySettings* MaterialShaderQualitySettings = UMaterialShaderQualitySettings::Get();
+				if (MaterialShaderQualitySettings->GetPreviewPlatform() != PreviewPlatformName)
+				{
+					return;
+				}
+				
+				SetPreviewPlatform(FName(), ERHIFeatureLevel::SM5);
+			}
+		});
+		
 	bNotifyUndoRedoSelectionChange = true;
 
 	EditorWorldExtensionsManager = nullptr;
diff --git a/Engine/Source/Runtime/Android/AndroidRuntimeSettings/Classes/AndroidRuntimeSettings.h b/Engine/Source/Runtime/Android/AndroidRuntimeSettings/Classes/AndroidRuntimeSettings.h
--- a/Engine/Source/Runtime/Android/AndroidRuntimeSettings/Classes/AndroidRuntimeSettings.h
+++ b/Engine/Source/Runtime/Android/AndroidRuntimeSettings/Classes/AndroidRuntimeSettings.h
@@ -580,6 +580,10 @@ public:
 
 
 #if WITH_EDITOR
+	/** Called whenever a registered Android property changes. */
+	DECLARE_MULTICAST_DELEGATE_OneParam(FPropertyChanged, struct FPropertyChangedEvent&);
+	FPropertyChanged OnPropertyChanged;
+
 private:
 	// UObject interface
 	virtual void PostEditChangeProperty(struct FPropertyChangedEvent& PropertyChangedEvent) override;
diff --git a/Engine/Source/Runtime/Android/AndroidRuntimeSettings/Private/AndroidRuntimeSettings.cpp b/Engine/Source/Runtime/Android/AndroidRuntimeSettings/Private/AndroidRuntimeSettings.cpp
--- a/Engine/Source/Runtime/Android/AndroidRuntimeSettings/Private/AndroidRuntimeSettings.cpp
+++ b/Engine/Source/Runtime/Android/AndroidRuntimeSettings/Private/AndroidRuntimeSettings.cpp
@@ -108,6 +108,8 @@ void UAndroidRuntimeSettings::PostEditChangeProperty(struct FPropertyChangedEven
 		{
 			// Supported shader formats changed so invalidate cache
 			InvalidateAllAndroidPlatforms();
+
+			OnPropertyChanged.Broadcast(PropertyChangedEvent);
 		}
 	}
 
diff --git a/Engine/Source/Runtime/Core/Private/Misc/CoreDelegates.cpp b/Engine/Source/Runtime/Core/Private/Misc/CoreDelegates.cpp
--- a/Engine/Source/Runtime/Core/Private/Misc/CoreDelegates.cpp
+++ b/Engine/Source/Runtime/Core/Private/Misc/CoreDelegates.cpp
@@ -124,6 +124,7 @@ FCoreDelegates::FOnSystemResolutionChanged FCoreDelegates::OnSystemResolutionCha
 
 #if WITH_EDITOR
 FCoreDelegates::FOnTargetPlatformChangedSupportedFormats FCoreDelegates::OnTargetPlatformChangedSupportedFormats;
+FCoreDelegates::FOnFeatureLevelDisabled FCoreDelegates::OnFeatureLevelDisabled;
 #endif 
 
 FCoreDelegates::FStatCheckEnabled FCoreDelegates::StatCheckEnabled;
diff --git a/Engine/Source/Runtime/Core/Public/Misc/CoreDelegates.h b/Engine/Source/Runtime/Core/Public/Misc/CoreDelegates.h
--- a/Engine/Source/Runtime/Core/Public/Misc/CoreDelegates.h
+++ b/Engine/Source/Runtime/Core/Public/Misc/CoreDelegates.h
@@ -302,6 +302,10 @@ public:
 	// called when a target platform changes it's return value of supported formats.  This is so anything caching those results can reset (like cached shaders for cooking)
 	DECLARE_MULTICAST_DELEGATE_OneParam(FOnTargetPlatformChangedSupportedFormats, const ITargetPlatform*); 
 	static FOnTargetPlatformChangedSupportedFormats OnTargetPlatformChangedSupportedFormats;
+
+	// Called when a feature level is disabled by the user.
+	DECLARE_MULTICAST_DELEGATE_TwoParams(FOnFeatureLevelDisabled, int, const FName&);
+	static FOnFeatureLevelDisabled OnFeatureLevelDisabled;
 #endif
 
 	/** IOS-style application lifecycle delegates */
