Remove extraneous line causing compilation issues

diff --git a/Engine/Source/Developer/Android/AndroidPlatformEditor/Private/AndroidPlatformEditorModule.cpp b/Engine/Source/Developer/Android/AndroidPlatformEditor/Private/AndroidPlatformEditorModule.cpp
--- a/Engine/Source/Developer/Android/AndroidPlatformEditor/Private/AndroidPlatformEditorModule.cpp
+++ b/Engine/Source/Developer/Android/AndroidPlatformEditor/Private/AndroidPlatformEditorModule.cpp
@@ -113,7 +113,6 @@ class FAndroidPlatformEditorModule
 		{
 			AndroidRuntimeSettings->OnPropertyChanged.AddLambda([this, AndroidRuntimeSettings](struct FPropertyChangedEvent& PropertyChangedEvent)
 				{
-					UEditorEngine* Editor = (UEditorEngine*)GEngine;
 					if (PropertyChangedEvent.Property != nullptr)
 					{
 						if (PropertyChangedEvent.Property->GetFName() == GET_MEMBER_NAME_CHECKED(UAndroidRuntimeSettings, bBuildForES31) &&
