Fix up inconsistencies between UE4/Main and Dev-EngineMerge
#rnx

diff --git a/Engine/Shaders/Private/GlobalDistanceFieldShared.ush b/Engine/Shaders/Private/GlobalDistanceFieldShared.ush
--- a/Engine/Shaders/Private/GlobalDistanceFieldShared.ush
+++ b/Engine/Shaders/Private/GlobalDistanceFieldShared.ush
@@ -1,7 +1,7 @@
 // Copyright 1998-2019 Epic Games, Inc. All Rights Reserved.
 
 /*=============================================================================
-	GlobalDistanceFieldShared.usf
+	GlobalDistanceFieldShared.ush
 =============================================================================*/
 
 #if FEATURE_LEVEL >= FEATURE_LEVEL_SM5
diff --git a/Engine/Source/Developer/Android/AndroidPlatformEditor/Private/AndroidLicenseDialog.cpp b/Engine/Source/Developer/Android/AndroidPlatformEditor/Private/AndroidLicenseDialog.cpp
--- a/Engine/Source/Developer/Android/AndroidPlatformEditor/Private/AndroidLicenseDialog.cpp
+++ b/Engine/Source/Developer/Android/AndroidPlatformEditor/Private/AndroidLicenseDialog.cpp
@@ -85,12 +85,10 @@ void SAndroidLicenseDialog::Construct(const FArguments& InArgs)
 			{
 				int32 LicenseLength = LicenseEnd - LicenseStart;
 
-				int32 ConvertedLength = FUTF8ToTCHAR_Convert::ConvertedLength((ANSICHAR*)LicenseStart, LicenseLength);
-				TCHAR* ConvertedBuffer = (TCHAR*)FMemory::Malloc((ConvertedLength + 1) * sizeof(TCHAR));
-				FUTF8ToTCHAR_Convert::Convert(ConvertedBuffer, ConvertedLength, (ANSICHAR*)LicenseStart, LicenseLength);
-				ConvertedBuffer[ConvertedLength] = 0;
-				LicenseText = FString(ConvertedBuffer);
-				FMemory::Free(ConvertedBuffer);
+				{
+					const FUTF8ToTCHAR ConvertedString(reinterpret_cast<ANSICHAR*>(LicenseStart), LicenseLength);
+					LicenseText = FString(ConvertedString.Length(), ConvertedString.Get());
+				}
 
 				FSHA1::HashBuffer(LicenseStart, LicenseLength, LicenseHash.Hash);
 
diff --git a/Engine/Source/Runtime/CoreUObject/Public/UObject/UObjectClusters.h b/Engine/Source/Runtime/CoreUObject/Public/UObject/UObjectClusters.h
--- a/Engine/Source/Runtime/CoreUObject/Public/UObject/UObjectClusters.h
+++ b/Engine/Source/Runtime/CoreUObject/Public/UObject/UObjectClusters.h
@@ -7,7 +7,6 @@
 #pragma once
 
 #include "CoreMinimal.h"
-#include "UObject/UObjectArray.h"
 
 #ifndef UE_GCCLUSTER_VERBOSE_LOGGING
 #define UE_GCCLUSTER_VERBOSE_LOGGING (0 && !UE_BUILD_SHIPPING)
@@ -16,7 +15,7 @@
 #if !UE_BUILD_SHIPPING
 
 // Dumps a single cluster to log
-COREUOBJECT_API void DumpClusterToLog(const FUObjectCluster& Cluster, bool bHierarchy, bool bIndexOnly);
+COREUOBJECT_API void DumpClusterToLog(const struct FUObjectCluster& Cluster, bool bHierarchy, bool bIndexOnly);
 
 // Dumps all clusters to log.
 COREUOBJECT_API void ListClusters(const TArray<FString>& Args);
diff --git a/Engine/Source/Runtime/Engine/Private/GameplayStatics.cpp b/Engine/Source/Runtime/Engine/Private/GameplayStatics.cpp
--- a/Engine/Source/Runtime/Engine/Private/GameplayStatics.cpp
+++ b/Engine/Source/Runtime/Engine/Private/GameplayStatics.cpp
@@ -1232,7 +1232,7 @@ void UGameplayStatics::SetGlobalListenerFocusParameters(const UObject* WorldCont
 	}
 }
 
-void UGameplayStatics::PlaySound2D(const UObject* WorldContextObject, class USoundBase* Sound, float VolumeMultiplier, float PitchMultiplier, float StartTime, class USoundConcurrency* ConcurrencySettings, AActor* OwningActor)
+void UGameplayStatics::PlaySound2D(const UObject* WorldContextObject, USoundBase* Sound, float VolumeMultiplier, float PitchMultiplier, float StartTime, USoundConcurrency* ConcurrencySettings, AActor* OwningActor)
 {
 	if (!Sound || !GEngine || !GEngine->UseSound())
 	{
