Hide the x86 setting properly on Android

#ue4
#android

diff --git a/Engine/Source/Developer/Android/AndroidPlatformEditor/Private/AndroidTargetSettingsCustomization.cpp b/Engine/Source/Developer/Android/AndroidPlatformEditor/Private/AndroidTargetSettingsCustomization.cpp
--- a/Engine/Source/Developer/Android/AndroidPlatformEditor/Private/AndroidTargetSettingsCustomization.cpp
+++ b/Engine/Source/Developer/Android/AndroidPlatformEditor/Private/AndroidTargetSettingsCustomization.cpp
@@ -299,7 +299,7 @@ void FAndroidTargetSettingsCustomization::BuildAppManifestSection(IDetailLayoutB
 	}
 	SETUP_ANDROIDARCH_PROP(TEXT("-armv7"), bBuildForArmV7, BuildCategory, LOCTEXT("BuildForArmV7ToolTip", "Enable ArmV7 CPU architecture support? (this will be used if all CPU architecture types are unchecked)"));
 	SETUP_ANDROIDARCH_PROP(TEXT("-arm64"), bBuildForArm64, BuildCategory, LOCTEXT("BuildForArm64ToolTip", "Enable Arm64 CPU architecture support? (use at least NDK r11c, requires Lollipop (android-21) minimum)"));
-	SETUP_ANDROIDARCH_PROP(TEXT("-x86"), bBuildForX86, BuildCategory, LOCTEXT("BuildForX86ToolTip", "Enable X86 CPU architecture support?"));
+//	SETUP_ANDROIDARCH_PROP(TEXT("-x86"), bBuildForX86, BuildCategory, LOCTEXT("BuildForX86ToolTip", "Enable X86 CPU architecture support?"));
 	SETUP_ANDROIDARCH_PROP(TEXT("-x64"), bBuildForX8664, BuildCategory, LOCTEXT("BuildForX8664ToolTip", "Enable X86-64 CPU architecture support?"));
 	SETUP_ANDROIDARCH_PROP(TEXT("-es2"), bBuildForES2, BuildCategory, LOCTEXT("BuildForES2ToolTip", "Enable OpenGL ES2 rendering support? (this will be used if rendering types are unchecked)"));
 
diff --git a/Engine/Source/Runtime/Android/AndroidRuntimeSettings/Private/AndroidRuntimeSettings.cpp b/Engine/Source/Runtime/Android/AndroidRuntimeSettings/Private/AndroidRuntimeSettings.cpp
--- a/Engine/Source/Runtime/Android/AndroidRuntimeSettings/Private/AndroidRuntimeSettings.cpp
+++ b/Engine/Source/Runtime/Android/AndroidRuntimeSettings/Private/AndroidRuntimeSettings.cpp
@@ -59,7 +59,7 @@ void UAndroidRuntimeSettings::PostEditChangeProperty(struct FPropertyChangedEven
 	Super::PostEditChangeProperty(PropertyChangedEvent);
 
 	// Ensure that at least one architecture is supported
-	if (!bBuildForArmV7 && !bBuildForX86 && !bBuildForX8664 && !bBuildForArm64)
+	if (!bBuildForArmV7 && !bBuildForX8664 && !bBuildForArm64)
 	{
 		bBuildForArmV7 = true;
 		UpdateSinglePropertyInConfigFile(GetClass()->FindPropertyByName(GET_MEMBER_NAME_CHECKED(UAndroidRuntimeSettings, bBuildForArmV7)), GetDefaultConfigFilename());
