non-unity build fix

diff --git a/Engine/Source/Developer/Android/AndroidTargetPlatform/Private/AndroidTargetPlatform.cpp b/Engine/Source/Developer/Android/AndroidTargetPlatform/Private/AndroidTargetPlatform.cpp
--- a/Engine/Source/Developer/Android/AndroidTargetPlatform/Private/AndroidTargetPlatform.cpp
+++ b/Engine/Source/Developer/Android/AndroidTargetPlatform/Private/AndroidTargetPlatform.cpp
@@ -21,11 +21,13 @@
 #include "Misc/SecureHash.h"
 #include "HAL/FileManager.h"
 #include "HAL/PlatformFilemanager.h"
+#include "HAL/IConsoleManager.h"
 #include "Interfaces/IAndroidDeviceDetectionModule.h"
 #include "Interfaces/IAndroidDeviceDetection.h"
 #include "Modules/ModuleManager.h"
 #include "Misc/SecureHash.h"
 
+
 #if WITH_ENGINE
 #include "AudioCompressionSettings.h"
 #endif
diff --git a/Engine/Source/Developer/IOS/IOSTargetPlatform/Private/IOSTargetPlatform.cpp b/Engine/Source/Developer/IOS/IOSTargetPlatform/Private/IOSTargetPlatform.cpp
--- a/Engine/Source/Developer/IOS/IOSTargetPlatform/Private/IOSTargetPlatform.cpp
+++ b/Engine/Source/Developer/IOS/IOSTargetPlatform/Private/IOSTargetPlatform.cpp
@@ -8,6 +8,7 @@
 #include "Interfaces/IProjectManager.h"
 #include "InstalledPlatformInfo.h"
 #include "HAL/FileManager.h"
+#include "HAL/IConsoleManager.h"
 #include "Misc/Paths.h"
 #include "Misc/App.h"
 #include "Misc/MonitoredProcess.h"
