Fix static analysis warning, based on CL 4485257 from BogdanV

diff --git a/Engine/Source/Developer/Android/AndroidDeviceDetection/Private/AndroidDeviceDetectionModule.cpp b/Engine/Source/Developer/Android/AndroidDeviceDetection/Private/AndroidDeviceDetectionModule.cpp
--- a/Engine/Source/Developer/Android/AndroidDeviceDetection/Private/AndroidDeviceDetectionModule.cpp
+++ b/Engine/Source/Developer/Android/AndroidDeviceDetection/Private/AndroidDeviceDetectionModule.cpp
@@ -689,7 +689,7 @@ public:
 			if (bOpenGL3x)
 			{
 				DeviceSpecs.AndroidProperties.GLES31RHIState.SupportsInstancing = true;
-				DeviceSpecs.AndroidProperties.GLES31RHIState.MaxTextureDimensions = 2048;
+				DeviceSpecs.AndroidProperties.GLES31RHIState.MaxTextureDimensions = 4096;
 				DeviceSpecs.AndroidProperties.GLES31RHIState.MaxShadowDepthBufferSizeX = 2048;
 				DeviceSpecs.AndroidProperties.GLES31RHIState.MaxShadowDepthBufferSizeY = 2048;
 				DeviceSpecs.AndroidProperties.GLES31RHIState.MaxCubeTextureDimensions = 2048;
@@ -701,7 +701,7 @@ public:
 			// OpenGL ES 2.0
 			{
 				DeviceSpecs.AndroidProperties.GLES2RHIState.SupportsInstancing = false;
-				DeviceSpecs.AndroidProperties.GLES2RHIState.MaxTextureDimensions = 1024;
+				DeviceSpecs.AndroidProperties.GLES2RHIState.MaxTextureDimensions = 2048;
 				DeviceSpecs.AndroidProperties.GLES2RHIState.MaxShadowDepthBufferSizeX = 1024;
 				DeviceSpecs.AndroidProperties.GLES2RHIState.MaxShadowDepthBufferSizeY = 1024;
 				DeviceSpecs.AndroidProperties.GLES2RHIState.MaxCubeTextureDimensions = 512;
diff --git a/Engine/Source/Editor/PIEPreviewDeviceProfileSelector/Private/PIEPreviewDevice.cpp b/Engine/Source/Editor/PIEPreviewDeviceProfileSelector/Private/PIEPreviewDevice.cpp
--- a/Engine/Source/Editor/PIEPreviewDeviceProfileSelector/Private/PIEPreviewDevice.cpp
+++ b/Engine/Source/Editor/PIEPreviewDeviceProfileSelector/Private/PIEPreviewDevice.cpp
@@ -297,7 +297,10 @@ void FPIEPreviewDevice::SetupDevice(const int32 InWindowTitleBarSize)
 
 	// load the bezel texture
 	BezelTexture = FImageUtils::ImportFileAsTexture2D(BezelPath);
-	BezelTexture->AddToRoot();
+	if (BezelTexture != nullptr)
+	{
+		BezelTexture->AddToRoot();
+	}
 	
 	// if we have invalid/uninitialized viewport values use the values provided as native device resolution
 	FPIEPreviewDeviceBezelViewportRect& ViewportRect = DeviceSpecs->BezelProperties.BezelViewportRect;
diff --git a/Engine/Source/Editor/PIEPreviewDeviceProfileSelector/Private/PIEPreviewDeviceProfileSelectorModule.cpp b/Engine/Source/Editor/PIEPreviewDeviceProfileSelector/Private/PIEPreviewDeviceProfileSelectorModule.cpp
--- a/Engine/Source/Editor/PIEPreviewDeviceProfileSelector/Private/PIEPreviewDeviceProfileSelectorModule.cpp
+++ b/Engine/Source/Editor/PIEPreviewDeviceProfileSelector/Private/PIEPreviewDeviceProfileSelectorModule.cpp
@@ -75,11 +75,11 @@ void FPIEPreviewDeviceModule::InitPreviewDevice()
 	// to finish setup we need complete engine initialization
 	ViewportCreatedDelegate = UGameViewportClient::OnViewportCreated().AddRaw(this, &FPIEPreviewDeviceModule::OnViewportCreated);
 
-	if (ReadDeviceSpecification())
-	{
-		Device->ApplyRHIPrerequisitesOverrides();
-		DeviceProfile = Device->GetProfile();
-	}
+	bool bReadSuccess = ReadDeviceSpecification();
+	checkf(bReadSuccess, TEXT("Unable to read device specifications"));
+
+	Device->ApplyRHIPrerequisitesOverrides();
+	DeviceProfile = Device->GetProfile();
 }
 
 void FPIEPreviewDeviceModule::OnEngineInitComplete()
diff --git a/Engine/Source/Editor/UnrealEd/Private/Kismet2/DebuggerCommands.cpp b/Engine/Source/Editor/UnrealEd/Private/Kismet2/DebuggerCommands.cpp
--- a/Engine/Source/Editor/UnrealEd/Private/Kismet2/DebuggerCommands.cpp
+++ b/Engine/Source/Editor/UnrealEd/Private/Kismet2/DebuggerCommands.cpp
@@ -703,7 +703,7 @@ static void AddAndroidConfigExportSubMenus(FMenuBuilder& InMenuBuilder)
 	{
 		FAndroidDeviceInfo& DeviceInfo = Pair.Value;
 
-		FString ModelName = DeviceInfo.Model + TEXT("(") + DeviceInfo.DeviceBrand + TEXT(")");
+		FString ModelName = DeviceInfo.Model + TEXT("[") + DeviceInfo.DeviceBrand + TEXT("]");
 
 		// lambda function called to open the save dialog and trigger device export
 		auto LambdaSaveConfigFile = [DeviceName = Pair.Key, DefaultFileName = ModelName, DeviceDetection]()
