LNK fix - fixing up more fallout from the recent merge from DevMain. Adding in missing function that was dropped in the merge (matching Main's version).

diff --git a/Engine/Source/Developer/Android/AndroidTargetPlatform/Private/AndroidTargetPlatform.cpp b/Engine/Source/Developer/Android/AndroidTargetPlatform/Private/AndroidTargetPlatform.cpp
--- a/Engine/Source/Developer/Android/AndroidTargetPlatform/Private/AndroidTargetPlatform.cpp
+++ b/Engine/Source/Developer/Android/AndroidTargetPlatform/Private/AndroidTargetPlatform.cpp
@@ -812,4 +812,9 @@ bool FAndroidTargetPlatform::HandleTicker( float DeltaTime )
 	return true;
 }
 
+FAndroidTargetDeviceRef FAndroidTargetPlatform::CreateNewDevice(const FAndroidDeviceInfo &DeviceInfo)
+{
+	return MakeShareable(new FAndroidTargetDevice(*this, DeviceInfo.SerialNumber, GetAndroidVariantName()));
+}
+
 #undef LOCTEXT_NAMESPACE
