Mark ES2 as deprecated, use ES31 as default for Android

diff --git a/Engine/Config/BaseEngine.ini b/Engine/Config/BaseEngine.ini
--- a/Engine/Config/BaseEngine.ini
+++ b/Engine/Config/BaseEngine.ini
@@ -1938,7 +1938,8 @@ bBuildForArmV7=true
 bBuildForArm64=false
 bBuildForX86=false
 bBuildForX8664=false
-bBuildForES2=true
+bBuildForES2=false
+bBuildForES31=true
 bSupportsVulkan=false
 bDetectVulkanByDefault=true
 bSplitIntoSeparateApks=false
diff --git a/Engine/Source/Developer/Android/AndroidPlatformEditor/Private/AndroidTargetSettingsCustomization.cpp b/Engine/Source/Developer/Android/AndroidPlatformEditor/Private/AndroidTargetSettingsCustomization.cpp
--- a/Engine/Source/Developer/Android/AndroidPlatformEditor/Private/AndroidTargetSettingsCustomization.cpp
+++ b/Engine/Source/Developer/Android/AndroidPlatformEditor/Private/AndroidTargetSettingsCustomization.cpp
@@ -301,7 +301,6 @@ void FAndroidTargetSettingsCustomization::BuildAppManifestSection(IDetailLayoutB
 	SETUP_ANDROIDARCH_PROP(TEXT("-arm64"), bBuildForArm64, BuildCategory, LOCTEXT("BuildForArm64ToolTip", "Enable Arm64 CPU architecture support? (use at least NDK r11c, requires Lollipop (android-21) minimum)"));
 //	SETUP_ANDROIDARCH_PROP(TEXT("-x86"), bBuildForX86, BuildCategory, LOCTEXT("BuildForX86ToolTip", "Enable X86 CPU architecture support?"));
 	SETUP_ANDROIDARCH_PROP(TEXT("-x64"), bBuildForX8664, BuildCategory, LOCTEXT("BuildForX8664ToolTip", "Enable X86-64 CPU architecture support?"));
-	SETUP_ANDROIDARCH_PROP(TEXT("-es2"), bBuildForES2, BuildCategory, LOCTEXT("BuildForES2ToolTip", "Enable OpenGL ES2 rendering support? (this will be used if rendering types are unchecked)"));
 
 	// @todo android fat binary: Put back in when we expose those
 //	SETUP_SOURCEONLY_PROP(bSplitIntoSeparateApks, BuildCategory, LOCTEXT("SplitIntoSeparateAPKsToolTip", "If checked, CPU architectures and rendering types will be split into separate .apk files"));
diff --git a/Engine/Source/Runtime/Android/AndroidRuntimeSettings/Classes/AndroidRuntimeSettings.h b/Engine/Source/Runtime/Android/AndroidRuntimeSettings/Classes/AndroidRuntimeSettings.h
--- a/Engine/Source/Runtime/Android/AndroidRuntimeSettings/Classes/AndroidRuntimeSettings.h
+++ b/Engine/Source/Runtime/Android/AndroidRuntimeSettings/Classes/AndroidRuntimeSettings.h
@@ -379,15 +379,15 @@ public:
 	UPROPERTY(GlobalConfig, EditAnywhere, Category = Build, meta = (DisplayName = "Support x86_64 [aka x64]"))
 	bool bBuildForX8664;
 
-	// Enable ES2 support?
-	UPROPERTY(GlobalConfig, EditAnywhere, Category = Build, meta = (DisplayName = "Support OpenGL ES2"))
+	// Include shaders for devices that support OpenGL ES 2 and above
+	UPROPERTY(GlobalConfig, EditAnywhere, Category = Build, meta = (DisplayName = "Support OpenGL ES2 (Deprecated)"))
 	bool bBuildForES2;
 
-	// Enable ES3.1 support?
+	// Include shaders for devices supporting OpenGL ES 3.1 and above (default)
 	UPROPERTY(GlobalConfig, EditAnywhere, Category = Build, meta = (DisplayName = "Support OpenGL ES3.1"))
 	bool bBuildForES31;
 
-	// Enable Vulkan rendering support?
+	// Support the Vulkan RHI and include Vulkan shaders
 	UPROPERTY(GlobalConfig, EditAnywhere, Category = Build, meta = (DisplayName = "Support Vulkan"))
 	bool bSupportsVulkan;
 
diff --git a/Engine/Source/Runtime/Android/AndroidRuntimeSettings/Private/AndroidRuntimeSettings.cpp b/Engine/Source/Runtime/Android/AndroidRuntimeSettings/Private/AndroidRuntimeSettings.cpp
--- a/Engine/Source/Runtime/Android/AndroidRuntimeSettings/Private/AndroidRuntimeSettings.cpp
+++ b/Engine/Source/Runtime/Android/AndroidRuntimeSettings/Private/AndroidRuntimeSettings.cpp
@@ -210,8 +210,8 @@ void UAndroidRuntimeSettings::EnsureValidGPUArch()
 	// Ensure that at least one GPU architecture is supported
 	if (!bBuildForES2 && !bSupportsVulkan && !bBuildForES31)
 	{
-		bBuildForES2 = true;
-		UpdateSinglePropertyInConfigFile(GetClass()->FindPropertyByName(GET_MEMBER_NAME_CHECKED(UAndroidRuntimeSettings, bBuildForES2)), GetDefaultConfigFilename());
+		bBuildForES31 = true;
+		UpdateSinglePropertyInConfigFile(GetClass()->FindPropertyByName(GET_MEMBER_NAME_CHECKED(UAndroidRuntimeSettings, bBuildForES31)), GetDefaultConfigFilename());
 
 		// Supported shader formats changed so invalidate cache
 		InvalidateAllAndroidPlatforms();
