[NO NON-PII COMMIT MESSAGE]

diff --git a/Engine/Source/Developer/Android/AndroidTargetPlatform/Private/AndroidTargetPlatform.cpp b/Engine/Source/Developer/Android/AndroidTargetPlatform/Private/AndroidTargetPlatform.cpp
--- a/Engine/Source/Developer/Android/AndroidTargetPlatform/Private/AndroidTargetPlatform.cpp
+++ b/Engine/Source/Developer/Android/AndroidTargetPlatform/Private/AndroidTargetPlatform.cpp
@@ -774,7 +774,6 @@ FPlatformAudioCookOverrides* FAndroidTargetPlatform::GetAudioCompressionSettings
 {
 	static FPlatformAudioCookOverrides Settings;
 
-#if !WITH_EDITOR
 	static bool bCachedPlatformSettings = false;
 
 	if (!bCachedPlatformSettings)
@@ -782,9 +781,6 @@ FPlatformAudioCookOverrides* FAndroidTargetPlatform::GetAudioCompressionSettings
 		Android::CachePlatformAudioCookOverrides(Settings);
 		bCachedPlatformSettings = true;
 	}
-#else
-	Android::CachePlatformAudioCookOverrides(Settings);
-#endif
 
 	return &Settings;
 }
diff --git a/Engine/Source/Developer/IOS/IOSTargetPlatform/Private/IOSTargetPlatform.cpp b/Engine/Source/Developer/IOS/IOSTargetPlatform/Private/IOSTargetPlatform.cpp
--- a/Engine/Source/Developer/IOS/IOSTargetPlatform/Private/IOSTargetPlatform.cpp
+++ b/Engine/Source/Developer/IOS/IOSTargetPlatform/Private/IOSTargetPlatform.cpp
@@ -845,7 +845,6 @@ FPlatformAudioCookOverrides* FIOSTargetPlatform::GetAudioCompressionSettings() c
 {
 	static FPlatformAudioCookOverrides Settings;
 
-#if !WITH_EDITOR
 	static bool bCachedPlatformSettings = false;
 
 	if (!bCachedPlatformSettings)
@@ -853,9 +852,6 @@ FPlatformAudioCookOverrides* FIOSTargetPlatform::GetAudioCompressionSettings() c
 		IOS::CachePlatformAudioCookOverrides(Settings);
 		bCachedPlatformSettings = true;
 	}
-#else
-	IOS::CachePlatformAudioCookOverrides(Settings);
-#endif
 
 	return &Settings;
 }
diff --git a/Engine/Source/Developer/Windows/WindowsTargetPlatform/Private/GenericWindowsTargetPlatform.h b/Engine/Source/Developer/Windows/WindowsTargetPlatform/Private/GenericWindowsTargetPlatform.h
--- a/Engine/Source/Developer/Windows/WindowsTargetPlatform/Private/GenericWindowsTargetPlatform.h
+++ b/Engine/Source/Developer/Windows/WindowsTargetPlatform/Private/GenericWindowsTargetPlatform.h
@@ -558,7 +558,6 @@ public:
 	{
 		static FPlatformAudioCookOverrides Settings;
 
-#if !WITH_EDITOR
 		static bool bCachedPlatformSettings = false;
 
 		if (!bCachedPlatformSettings)
@@ -566,9 +565,6 @@ public:
 			Windows::CachePlatformAudioCookOverrides(Settings);
 			bCachedPlatformSettings = true;
 		}
-#else
-		Windows::CachePlatformAudioCookOverrides(Settings);
-#endif
 
 		return &Settings;
 	}
