Fixed post merge compilation issues.

diff --git a/Engine/Source/Developer/Android/AndroidTargetPlatform/Private/AndroidTargetPlatform.cpp b/Engine/Source/Developer/Android/AndroidTargetPlatform/Private/AndroidTargetPlatform.cpp
--- a/Engine/Source/Developer/Android/AndroidTargetPlatform/Private/AndroidTargetPlatform.cpp
+++ b/Engine/Source/Developer/Android/AndroidTargetPlatform/Private/AndroidTargetPlatform.cpp
@@ -424,21 +424,11 @@ void FAndroidTargetPlatform::GetTextureFormats( const UTexture* InTexture, TArra
 		TArray<FName> FormatPerLayer;
 		FormatPerLayer.SetNum(NumLayers);
 
->>>> ORIGINAL //UE4/Main/Engine/Source/Developer/Android/AndroidTargetPlatform/Private/AndroidTargetPlatform.cpp#15
-	// The order we add texture formats to OutFormats is important. When multiple formats are cooked
-	// and supported by the device, the first supported format listed will be used. 
-	// eg, ETC1/uncompressed should always be last
-==== THEIRS //UE4/Main/Engine/Source/Developer/Android/AndroidTargetPlatform/Private/AndroidTargetPlatform.cpp#16
-	// The order we add texture formats to OutFormats is important. When multiple formats are cooked
-	// and supported by the device, the first supported format listed will be used.
-	// eg, ETC1/uncompressed should always be last
-==== YOURS //Juan.Canada_Dev-Rendering-MERGES/Engine/Source/Developer/Android/AndroidTargetPlatform/Private/AndroidTargetPlatform.cpp
 		bool bValidFormat = true;
 		for (int32 LayerIndex = 0; LayerIndex < NumLayers; ++LayerIndex)
 		{
 			FTextureFormatSettings LayerFormatSettings;
 			InTexture->GetLayerFormatSettings(LayerIndex, LayerFormatSettings);
-<<<<
 
 			const bool bNoCompression = LayerFormatSettings.CompressionNone				// Code wants the texture uncompressed.
 				|| (InTexture->LODGroup == TEXTUREGROUP_ColorLookupTable)	// Textures in certain LOD groups should remain uncompressed.
diff --git a/Engine/Source/Runtime/D3D12RHI/Private/D3D12Adapter.cpp b/Engine/Source/Runtime/D3D12RHI/Private/D3D12Adapter.cpp
--- a/Engine/Source/Runtime/D3D12RHI/Private/D3D12Adapter.cpp
+++ b/Engine/Source/Runtime/D3D12RHI/Private/D3D12Adapter.cpp
@@ -102,18 +102,12 @@ void FD3D12Adapter::CreateRootDevice(bool bWithDebug)
 
 			bool bD3d12gpuvalidation = false;
 			if (FParse::Param(FCommandLine::Get(), TEXT("d3d12gpuvalidation")) || FParse::Param(FCommandLine::Get(), TEXT("gpuvalidation")))
+			{
 				TRefCountPtr<ID3D12Debug1> DebugController1;
 				VERIFYD3D12RESULT(DebugController->QueryInterface(IID_PPV_ARGS(DebugController1.GetInitReference())));
 				DebugController1->SetEnableGPUBasedValidation(true);
 				bD3d12gpuvalidation = true;
 			}
-			if (FParse::Param(FCommandLine::Get(), TEXT("d3d12gpuvalidation")) || FParse::Param(FCommandLine::Get(), TEXT("gpuvalidation")))
-		{
-			TRefCountPtr<ID3D12Debug1> DebugController1;
-			VERIFYD3D12RESULT(DebugController->QueryInterface(IID_PPV_ARGS(DebugController1.GetInitReference())));
-			DebugController1->SetEnableGPUBasedValidation(true);
-			bD3d12gpuvalidation = true;
-		}
 
 			UE_LOG(LogD3D12RHI, Log, TEXT("InitD3DDevice: -D3DDebug = %s -D3D12GPUValidation = %s"), bWithDebug ? TEXT("on") : TEXT("off"), bD3d12gpuvalidation ? TEXT("on") : TEXT("off"));
 		}
diff --git a/Engine/Source/Runtime/Renderer/Private/PostProcess/PostProcessTonemap.cpp b/Engine/Source/Runtime/Renderer/Private/PostProcess/PostProcessTonemap.cpp
--- a/Engine/Source/Runtime/Renderer/Private/PostProcess/PostProcessTonemap.cpp
+++ b/Engine/Source/Runtime/Renderer/Private/PostProcess/PostProcessTonemap.cpp
@@ -134,7 +134,7 @@ using FCommonDomain = TShaderPermutationDomain<
 	FTonemapperVignetteDim,
 	FTonemapperSharpenDim,
 	FTonemapperGrainJitterDim,
-	FTonemapperSwitchAxis>,
+	FTonemapperSwitchAxis,
 	FTonemapperUseVolumeLut>;
 
 FORCEINLINE_DEBUGGABLE bool ShouldCompileCommonPermutation(const FGlobalShaderPermutationParameters& Parameters, const FCommonDomain& PermutationVector)
