Removed ATC/ETC1 and Android PVR texture formats

diff --git a/Engine/Config/Android/DataDrivenPlatformInfo.ini b/Engine/Config/Android/DataDrivenPlatformInfo.ini
--- a/Engine/Config/Android/DataDrivenPlatformInfo.ini
+++ b/Engine/Config/Android/DataDrivenPlatformInfo.ini
@@ -42,25 +42,6 @@ bUATClosesAfterLaunch=true
 UBTTargetID=Android
 PlatformGroupName=Mobile
 
-[PlatformInfo Android_ATC]
-TargetPlatformName=Android_ATC
-DisplayName=Android (ATC)
-PlatformType=Game
-PlatformFlags=CookFlavor
-NormalIconPath=Launcher/Android/Platform_Android_ATC_24x
-LargeIconPath=Launcher/Android/Platform_Android_128x
-XLargeIconPath=
-UATCommandLine=-targetplatform=Android -cookflavor=ATC
-AutoSDKPath=Android
-TutorialPath=/Engine/Tutorial/Mobile/SettingUpAndroidTutorial.SettingUpAndroidTutorial
-bIsEnabled=true
-BinariesDirectoryName=Android
-IniPlatformName=Android
-bUsesHostCompiler=false
-bUATClosesAfterLaunch=true
-UBTTargetID=Android
-PlatformGroupName=Mobile
-
 [PlatformInfo Android_DXT]
 TargetPlatformName=Android_DXT
 DisplayName=Android (DXT)
@@ -80,44 +61,6 @@ bUATClosesAfterLaunch=true
 UBTTargetID=Android
 PlatformGroupName=Mobile
 
-[PlatformInfo Android_ETC1]
-TargetPlatformName=Android_ETC1
-DisplayName=Android (ETC1)
-PlatformType=Game
-PlatformFlags=CookFlavor
-NormalIconPath=Launcher/Android/Platform_Android_ETC1_24x
-LargeIconPath=Launcher/Android/Platform_Android_128x
-XLargeIconPath=
-UATCommandLine=-targetplatform=Android -cookflavor=ETC1
-AutoSDKPath=Android
-TutorialPath=/Engine/Tutorial/Mobile/SettingUpAndroidTutorial.SettingUpAndroidTutorial
-bIsEnabled=true
-BinariesDirectoryName=Android
-IniPlatformName=Android
-bUsesHostCompiler=false
-bUATClosesAfterLaunch=true
-UBTTargetID=Android
-PlatformGroupName=Mobile
-
-[PlatformInfo Android_ETC1a]
-TargetPlatformName=Android_ETC1a
-DisplayName=Android (ETC1a)
-PlatformType=Game
-PlatformFlags=CookFlavor
-NormalIconPath=Launcher/Android/Platform_Android_ETC1_24x
-LargeIconPath=Launcher/Android/Platform_Android_128x
-XLargeIconPath=
-UATCommandLine=-targetplatform=Android -cookflavor=ETC1a
-AutoSDKPath=Android
-TutorialPath=/Engine/Tutorial/Mobile/SettingUpAndroidTutorial.SettingUpAndroidTutorial
-bIsEnabled=true
-BinariesDirectoryName=Android
-IniPlatformName=Android
-bUsesHostCompiler=false
-bUATClosesAfterLaunch=true
-UBTTargetID=Android
-PlatformGroupName=Mobile
-
 [PlatformInfo Android_ETC2]
 TargetPlatformName=Android_ETC2
 DisplayName=Android (ETC2)
@@ -137,25 +80,6 @@ bUATClosesAfterLaunch=true
 UBTTargetID=Android
 PlatformGroupName=Mobile
 
-[PlatformInfo Android_PVRTC]
-TargetPlatformName=Android_PVRTC
-DisplayName=Android (PVRTC)
-PlatformType=Game
-PlatformFlags=CookFlavor
-NormalIconPath=Launcher/Android/Platform_Android_PVRTC_24x
-LargeIconPath=Launcher/Android/Platform_Android_128x
-XLargeIconPath=
-UATCommandLine=-targetplatform=Android -cookflavor=PVRTC
-AutoSDKPath=Android
-TutorialPath=/Engine/Tutorial/Mobile/SettingUpAndroidTutorial.SettingUpAndroidTutorial
-bIsEnabled=true
-BinariesDirectoryName=Android
-IniPlatformName=Android
-bUsesHostCompiler=false
-bUATClosesAfterLaunch=true
-UBTTargetID=Android
-PlatformGroupName=Mobile
-
 [PlatformInfo Android_ASTC]
 TargetPlatformName=Android_ASTC
 DisplayName=Android (ASTC)
@@ -214,25 +138,6 @@ bUATClosesAfterLaunch=true
 UBTTargetID=Android
 PlatformGroupName=Mobile
 
-[PlatformInfo Android_ATCClient]
-TargetPlatformName=Android_ATCClient
-DisplayName=Android (ATC) (Client-only)
-PlatformType=Client
-PlatformFlags=CookFlavor
-NormalIconPath=Launcher/Android/Platform_Android_ATC_24x
-LargeIconPath=Launcher/Android/Platform_Android_128x
-XLargeIconPath=
-UATCommandLine=-client -targetplatform=Android -cookflavor=ATC
-AutoSDKPath=Android
-TutorialPath=/Engine/Tutorial/Mobile/SettingUpAndroidTutorial.SettingUpAndroidTutorial
-bIsEnabled=true
-BinariesDirectoryName=Android
-IniPlatformName=Android
-bUsesHostCompiler=false
-bUATClosesAfterLaunch=true
-UBTTargetID=Android
-PlatformGroupName=Mobile
-
 [PlatformInfo Android_DXTClient]
 TargetPlatformName=Android_DXTClient
 DisplayName=Android (DXT) (Client-only)
@@ -252,44 +157,6 @@ bUATClosesAfterLaunch=true
 UBTTargetID=Android
 PlatformGroupName=Mobile
 
-[PlatformInfo Android_ETC1Client]
-TargetPlatformName=Android_ETC1Client
-DisplayName=Android (ETC1) (Client-only)
-PlatformType=Client
-PlatformFlags=CookFlavor
-NormalIconPath=Launcher/Android/Platform_Android_ETC1_24x
-LargeIconPath=Launcher/Android/Platform_Android_128x
-XLargeIconPath=
-UATCommandLine=-client -targetplatform=Android -cookflavor=ETC1
-AutoSDKPath=Android
-TutorialPath=/Engine/Tutorial/Mobile/SettingUpAndroidTutorial.SettingUpAndroidTutorial
-bIsEnabled=true
-BinariesDirectoryName=Android
-IniPlatformName=Android
-bUsesHostCompiler=false
-bUATClosesAfterLaunch=true
-UBTTargetID=Android
-PlatformGroupName=Mobile
-
-[PlatformInfo Android_ETC1aClient]
-TargetPlatformName=Android_ETC1aClient
-DisplayName=Android (ETC1a) (Client-only)
-PlatformType=Client
-PlatformFlags=CookFlavor
-NormalIconPath=Launcher/Android/Platform_Android_ETC1_24x
-LargeIconPath=Launcher/Android/Platform_Android_128x
-XLargeIconPath=
-UATCommandLine=-client -targetplatform=Android -cookflavor=ETC1a
-AutoSDKPath=Android
-TutorialPath=/Engine/Tutorial/Mobile/SettingUpAndroidTutorial.SettingUpAndroidTutorial
-bIsEnabled=true
-BinariesDirectoryName=Android
-IniPlatformName=Android
-bUsesHostCompiler=false
-bUATClosesAfterLaunch=true
-UBTTargetID=Android
-PlatformGroupName=Mobile
-
 [PlatformInfo Android_ETC2Client]
 TargetPlatformName=Android_ETC2Client
 DisplayName=Android (ETC2) (Client-only)
@@ -309,25 +176,6 @@ bUATClosesAfterLaunch=true
 UBTTargetID=Android
 PlatformGroupName=Mobile
 
-[PlatformInfo Android_PVRTCClient]
-TargetPlatformName=Android_PVRTCClient
-DisplayName=Android (PVRTC) (Client-only)
-PlatformType=Client
-PlatformFlags=CookFlavor
-NormalIconPath=Launcher/Android/Platform_Android_PVRTC_24x
-LargeIconPath=Launcher/Android/Platform_Android_128x
-XLargeIconPath=
-UATCommandLine=-client -targetplatform=Android -cookflavor=PVRTC
-AutoSDKPath=Android
-TutorialPath=/Engine/Tutorial/Mobile/SettingUpAndroidTutorial.SettingUpAndroidTutorial
-bIsEnabled=true
-BinariesDirectoryName=Android
-IniPlatformName=Android
-bUsesHostCompiler=false
-bUATClosesAfterLaunch=true
-UBTTargetID=Android
-PlatformGroupName=Mobile
-
 [PlatformInfo Android_ASTCClient]
 TargetPlatformName=Android_ASTCClient
 DisplayName=Android (ASTC) (Client-only)
diff --git a/Engine/Config/BaseEngine.ini b/Engine/Config/BaseEngine.ini
--- a/Engine/Config/BaseEngine.ini
+++ b/Engine/Config/BaseEngine.ini
@@ -2070,18 +2070,11 @@ bEnableGradle=true
 bEnableLint=false
 bShowLaunchImage=true
 bValidateTextureFormats=true
-bMultiTargetFormat_ETC1=true
 bMultiTargetFormat_ETC2=true
 bMultiTargetFormat_DXT=true
-bMultiTargetFormat_PVRTC=true
-bMultiTargetFormat_ATC=true
 bMultiTargetFormat_ASTC=true
-TextureFormatPriority_ETC1=0.1
-TextureFormatPriority_ETC1a=0.15
 TextureFormatPriority_ETC2=0.2
 TextureFormatPriority_DXT=0.6
-TextureFormatPriority_PVRTC=0.8
-TextureFormatPriority_ATC=0.5
 TextureFormatPriority_ASTC=0.9
 bEnableNewKeyboard=True
 bAndroidVoiceEnabled=false
diff --git a/Engine/Plugins/Editor/MobileLauncherProfileWizard/Source/MobileLauncherProfileWizard/Private/AndroidProfileWizard.cpp b/Engine/Plugins/Editor/MobileLauncherProfileWizard/Source/MobileLauncherProfileWizard/Private/AndroidProfileWizard.cpp
--- a/Engine/Plugins/Editor/MobileLauncherProfileWizard/Source/MobileLauncherProfileWizard/Private/AndroidProfileWizard.cpp
+++ b/Engine/Plugins/Editor/MobileLauncherProfileWizard/Source/MobileLauncherProfileWizard/Private/AndroidProfileWizard.cpp
@@ -23,7 +23,7 @@ FText FAndroidProfileWizard::GetDescription() const
 
 namespace AndroidProfileConstants
 {
-	const FString AppPlatformName(TEXT("Android_ETC1"));
+	const FString AppPlatformName(TEXT("Android_ETC2"));
 	const FString AppReleaseName(TEXT("1.0"));
 	const FString DLCName(TEXT("DLC1.0"));
 }
diff --git a/Engine/Source/Developer/Android/AndroidTargetPlatform/Private/AndroidTargetDevice.h b/Engine/Source/Developer/Android/AndroidTargetPlatform/Private/AndroidTargetDevice.h
--- a/Engine/Source/Developer/Android/AndroidTargetPlatform/Private/AndroidTargetDevice.h
+++ b/Engine/Source/Developer/Android/AndroidTargetPlatform/Private/AndroidTargetDevice.h
@@ -47,7 +47,7 @@ public:
 	 *
 	 * @param InTargetPlatform - The target platform.
 	 * @param InSerialNumber - The ADB serial number of the target device.
-	 * @param InAndroidVariant - The variant of the Android platform, i.e. ATC, DXT or PVRTC.
+	 * @param InAndroidVariant - The variant of the Android platform, i.e. ETC2, DXT or ASTC.
 	 */
 	FAndroidTargetDevice(const ITargetPlatform& InTargetPlatform, const FString& InSerialNumber, const FString& InAndroidVariant)
 		: AndroidVariant(InAndroidVariant)
@@ -229,7 +229,7 @@ protected:
 
 protected:
 
-	// The variant of the Android platform, i.e. ATC, DXT or PVRTC.
+	// The variant of the Android platform, i.e. ETC2, DXT or ASTC.
 	FString AndroidVariant;
 
 	// Holds a flag indicating whether the device is currently connected.
diff --git a/Engine/Source/Developer/Android/AndroidTargetPlatform/Private/AndroidTargetPlatform.cpp b/Engine/Source/Developer/Android/AndroidTargetPlatform/Private/AndroidTargetPlatform.cpp
--- a/Engine/Source/Developer/Android/AndroidTargetPlatform/Private/AndroidTargetPlatform.cpp
+++ b/Engine/Source/Developer/Android/AndroidTargetPlatform/Private/AndroidTargetPlatform.cpp
@@ -470,12 +470,8 @@ void FAndroidTargetPlatform::GetTextureFormats( const UTexture* InTexture, TArra
 			else if (LayerFormatSettings.CompressionSettings == TC_Normalmap)
 			{
 				if(!bIsCompressionValid) FormatPerLayer[LayerIndex] = AndroidTexFormat::NamePOTERROR;
-				else if(FormatCategory == EAndroidTextureFormatCategory::PVRTC) FormatPerLayer[LayerIndex] = AndroidTexFormat::NamePVRTC4;
 				else if (FormatCategory == EAndroidTextureFormatCategory::DXT) FormatPerLayer[LayerIndex] = AndroidTexFormat::NameDXT5;
-				else if (FormatCategory == EAndroidTextureFormatCategory::ATC) FormatPerLayer[LayerIndex] = AndroidTexFormat::NameATC_RGBA_I;
 				else if (FormatCategory == EAndroidTextureFormatCategory::ETC2) FormatPerLayer[LayerIndex] = AndroidTexFormat::NameETC2_RGB;
-				else if (FormatCategory == EAndroidTextureFormatCategory::ETC1a) FormatPerLayer[LayerIndex] = AndroidTexFormat::NameAutoETC1a;
-				else if (FormatCategory == EAndroidTextureFormatCategory::ETC1) FormatPerLayer[LayerIndex] = AndroidTexFormat::NameAutoETC1;
 				else bValidFormat = false;
 			}
 			else if (LayerFormatSettings.CompressionSettings == TC_Displacementmap)
@@ -498,48 +494,32 @@ void FAndroidTargetPlatform::GetTextureFormats( const UTexture* InTexture, TArra
 			{
 				FormatPerLayer[LayerIndex] = AndroidTexFormat::NameG8;
 			}
-			else if (InTexture->bForcePVRTC4 || LayerFormatSettings.CompressionSettings == TC_BC7)
+			else if (LayerFormatSettings.CompressionSettings == TC_BC7)
 			{
 				if (!bIsCompressionValid) FormatPerLayer[LayerIndex] = AndroidTexFormat::NamePOTERROR;
-				else if (FormatCategory == EAndroidTextureFormatCategory::PVRTC) FormatPerLayer[LayerIndex] = AndroidTexFormat::NamePVRTC4;
 				else if (FormatCategory == EAndroidTextureFormatCategory::DXT) FormatPerLayer[LayerIndex] = AndroidTexFormat::NameDXT5;
-				else if (FormatCategory == EAndroidTextureFormatCategory::ATC) FormatPerLayer[LayerIndex] = AndroidTexFormat::NameATC_RGBA_I;
 				else if (FormatCategory == EAndroidTextureFormatCategory::ETC2) FormatPerLayer[LayerIndex] = AndroidTexFormat::NameAutoETC2;
-				else if (FormatCategory == EAndroidTextureFormatCategory::ETC1a) FormatPerLayer[LayerIndex] = AndroidTexFormat::NameAutoETC1a;
-				else if (FormatCategory == EAndroidTextureFormatCategory::ETC1) FormatPerLayer[LayerIndex] = AndroidTexFormat::NameAutoETC1;
 				else bValidFormat = false;
 			}
 			else if (LayerFormatSettings.CompressionNoAlpha)
 			{
 				if (!bIsCompressionValid) FormatPerLayer[LayerIndex] = AndroidTexFormat::NamePOTERROR;
-				else if (FormatCategory == EAndroidTextureFormatCategory::PVRTC) FormatPerLayer[LayerIndex] = AndroidTexFormat::NamePVRTC2;
 				else if (FormatCategory == EAndroidTextureFormatCategory::DXT) FormatPerLayer[LayerIndex] = AndroidTexFormat::NameDXT1;
-				else if (FormatCategory == EAndroidTextureFormatCategory::ATC) FormatPerLayer[LayerIndex] = AndroidTexFormat::NameATC_RGB;
 				else if (FormatCategory == EAndroidTextureFormatCategory::ETC2) FormatPerLayer[LayerIndex] = AndroidTexFormat::NameETC2_RGB;
-				else if (FormatCategory == EAndroidTextureFormatCategory::ETC1a) FormatPerLayer[LayerIndex] = AndroidTexFormat::NameAutoETC1a;
-				else if (FormatCategory == EAndroidTextureFormatCategory::ETC1) FormatPerLayer[LayerIndex] = AndroidTexFormat::NameETC1;
 				else bValidFormat = false;
 			}
 			else if (InTexture->bDitherMipMapAlpha)
 			{
 				if (!bIsCompressionValid) FormatPerLayer[LayerIndex] = AndroidTexFormat::NamePOTERROR;
-				else if (FormatCategory == EAndroidTextureFormatCategory::PVRTC) FormatPerLayer[LayerIndex] = AndroidTexFormat::NamePVRTC4;
 				else if (FormatCategory == EAndroidTextureFormatCategory::DXT) FormatPerLayer[LayerIndex] = AndroidTexFormat::NameDXT5;
-				else if (FormatCategory == EAndroidTextureFormatCategory::ATC) FormatPerLayer[LayerIndex] = AndroidTexFormat::NameATC_RGBA_I;
 				else if (FormatCategory == EAndroidTextureFormatCategory::ETC2) FormatPerLayer[LayerIndex] = AndroidTexFormat::NameAutoETC2;
-				else if (FormatCategory == EAndroidTextureFormatCategory::ETC1a) FormatPerLayer[LayerIndex] = AndroidTexFormat::NameAutoETC1a;
-				else if (FormatCategory == EAndroidTextureFormatCategory::ETC1) FormatPerLayer[LayerIndex] = AndroidTexFormat::NameAutoETC1;
 				else bValidFormat = false;
 			}
 			else
 			{
 				if (!bIsCompressionValid) FormatPerLayer[LayerIndex] = AndroidTexFormat::NamePOTERROR;
-				else if (FormatCategory == EAndroidTextureFormatCategory::PVRTC) FormatPerLayer[LayerIndex] = AndroidTexFormat::NamePVRTC4;
 				else if (FormatCategory == EAndroidTextureFormatCategory::DXT) FormatPerLayer[LayerIndex] = AndroidTexFormat::NameAutoDXT;
-				else if (FormatCategory == EAndroidTextureFormatCategory::ATC) FormatPerLayer[LayerIndex] = AndroidTexFormat::NameAutoATC;
 				else if (FormatCategory == EAndroidTextureFormatCategory::ETC2) FormatPerLayer[LayerIndex] = AndroidTexFormat::NameAutoETC2;
-				else if (FormatCategory == EAndroidTextureFormatCategory::ETC1a) FormatPerLayer[LayerIndex] = AndroidTexFormat::NameAutoETC1a;
-				else if (FormatCategory == EAndroidTextureFormatCategory::ETC1) FormatPerLayer[LayerIndex] = AndroidTexFormat::NameAutoETC1;
 				else bValidFormat = false;
 			}
 		}
@@ -568,22 +548,10 @@ void FAndroidTargetPlatform::GetAllTextureFormats(TArray<FName>& OutFormats) con
 
 	auto AddAllTextureFormatIfSupports = [=, &OutFormats](bool bIsNonPOT)
 	{
-		AddTextureFormatIfSupports(AndroidTexFormat::NameAutoPVRTC, OutFormats, bIsNonPOT);
-		AddTextureFormatIfSupports(AndroidTexFormat::NamePVRTC2, OutFormats, bIsNonPOT);
-		AddTextureFormatIfSupports(AndroidTexFormat::NamePVRTC4, OutFormats, bIsNonPOT);
-
 		AddTextureFormatIfSupports(AndroidTexFormat::NameAutoDXT, OutFormats, bIsNonPOT);
 		AddTextureFormatIfSupports(AndroidTexFormat::NameDXT1, OutFormats, bIsNonPOT);
 		AddTextureFormatIfSupports(AndroidTexFormat::NameDXT5, OutFormats, bIsNonPOT);
-
-		AddTextureFormatIfSupports(AndroidTexFormat::NameATC_RGB, OutFormats, bIsNonPOT);
-		AddTextureFormatIfSupports(AndroidTexFormat::NameATC_RGBA_I, OutFormats, bIsNonPOT);
-
-		AddTextureFormatIfSupports(AndroidTexFormat::NameAutoETC1, OutFormats, bIsNonPOT);
-		AddTextureFormatIfSupports(AndroidTexFormat::NameAutoETC1a, OutFormats, bIsNonPOT);
 		AddTextureFormatIfSupports(AndroidTexFormat::NameAutoETC2, OutFormats, bIsNonPOT);
-
-		AddTextureFormatIfSupports(AndroidTexFormat::NameAutoATC, OutFormats, bIsNonPOT);
 	};
 
 	AddAllTextureFormatIfSupports(true);
diff --git a/Engine/Source/Developer/Android/AndroidTargetPlatform/Private/AndroidTargetPlatform.h b/Engine/Source/Developer/Android/AndroidTargetPlatform/Private/AndroidTargetPlatform.h
--- a/Engine/Source/Developer/Android/AndroidTargetPlatform/Private/AndroidTargetPlatform.h
+++ b/Engine/Source/Developer/Android/AndroidTargetPlatform/Private/AndroidTargetPlatform.h
@@ -46,19 +46,9 @@ template<typename TPlatformProperties> class TTargetPlatformBase;
 namespace AndroidTexFormat
 {
 	// Compressed Texture Formats
-	static FName NamePVRTC2(TEXT("PVRTC2"));
-	static FName NamePVRTC4(TEXT("PVRTC4"));
-	static FName NameAutoPVRTC(TEXT("AutoPVRTC"));
 	static FName NameDXT1(TEXT("DXT1"));
 	static FName NameDXT5(TEXT("DXT5"));
 	static FName NameAutoDXT(TEXT("AutoDXT"));
-	static FName NameATC_RGB(TEXT("ATC_RGB"));
-	static FName NameATC_RGBA_E(TEXT("ATC_RGBA_E"));		// explicit alpha
-	static FName NameATC_RGBA_I(TEXT("ATC_RGBA_I"));		// interpolated alpha
-	static FName NameAutoATC(TEXT("AutoATC"));
-	static FName NameETC1(TEXT("ETC1"));
-	static FName NameAutoETC1(TEXT("AutoETC1"));			// ETC1 or uncompressed RGBA, if alpha channel required
-	static FName NameAutoETC1a(TEXT("AutoETC1a"));
 	static FName NameETC2_RGB(TEXT("ETC2_RGB"));
 	static FName NameETC2_RGBA(TEXT("ETC2_RGBA"));
 	static FName NameAutoETC2(TEXT("AutoETC2"));
@@ -82,12 +72,8 @@ namespace AndroidTexFormat
 /** Listed in order of priority...if device supports multiple formats, first format in list will be chosen */
 enum class EAndroidTextureFormatCategory
 {
-	PVRTC,
 	DXT,
-	ATC,
 	ETC2,
-	ETC1a,
-	ETC1,
 	ASTC,
 
 	Count,
@@ -113,7 +99,7 @@ public:
 public:
 
 	/**
-	 * Gets the name of the Android platform variant, i.e. ATC, DXT, PVRTC, etc.
+	 * Gets the name of the Android platform variant, i.e. ASTC, ETC2, DXT, etc.
 	 *
 	 * @param Variant name.
 	 */
@@ -388,61 +374,6 @@ public:
 };
 
 
-
-class FAndroid_ATCTargetPlatform : public FAndroidTargetPlatform
-{
-public:
-	FAndroid_ATCTargetPlatform(bool bIsClient) : FAndroidTargetPlatform(bIsClient)
-	{
-		this->PlatformInfo = PlatformInfo::FindPlatformInfo("Android_ATC");
-	}
-
-	virtual FString GetAndroidVariantName() const override
-	{
-		return TEXT("ATC");
-	}
-
-	virtual FText DisplayName() const override
-	{
-		return LOCTEXT("Android_ATC", "Android (ATC)");
-	}
-
-	virtual bool SupportsTextureFormat(FName Format) const override
-	{
-		if (Format == AndroidTexFormat::NameATC_RGB ||
-			Format == AndroidTexFormat::NameATC_RGBA_I ||
-			Format == AndroidTexFormat::NameAutoATC)
-		{
-			return true;
-		}
-		return false;
-	}
-
-	virtual bool SupportsTextureFormatCategory(EAndroidTextureFormatCategory Category) const override
-	{
-		return Category == EAndroidTextureFormatCategory::ATC;
-	}
-
-	virtual bool SupportedByExtensionsString(const FString& ExtensionsString, const int GLESVersion) const override
-	{
-		return (ExtensionsString.Contains(TEXT("GL_ATI_texture_compression_atitc")) || ExtensionsString.Contains(TEXT("GL_AMD_compressed_ATC_texture")));
-	}
-
-	virtual FText GetVariantDisplayName() const override
-	{
-		return LOCTEXT("Android_ATC_ShortName", "ATC");
-	}
-
-	virtual float GetVariantPriority() const override
-	{
-		float Priority;
-		return (GConfig->GetFloat(TEXT("/Script/AndroidRuntimeSettings.AndroidRuntimeSettings"), TEXT("TextureFormatPriority_ATC"), Priority, GEngineIni) ?
-			Priority : 0.5f) * 10.0f + (IsClientOnly() ? 0.25f : 0.5f);
-	}
-};
-
-
-
 class FAndroid_ASTCTargetPlatform : public FAndroidTargetPlatform
 {
 public:
@@ -485,17 +416,17 @@ public:
 	{
 		check(Texture);
 
-		// we remap some of the defaults (with PVRTC and ASTC formats)
+		// we remap some of the defaults (with ASTC formats)
 		static FName FormatRemap[][2] =
 		{
 			// Default format:				ASTC format:
-			{ { FName(TEXT("DXT1")) },{ FName(TEXT("ASTC_RGB")) } },
-			{ { FName(TEXT("DXT5")) },{ FName(TEXT("ASTC_RGBA")) } },
-			{ { FName(TEXT("DXT5n")) },{ FName(TEXT("ASTC_NormalAG")) } },
-			{ { FName(TEXT("BC5")) },{ FName(TEXT("ASTC_NormalRG")) } },
-			{ { FName(TEXT("BC6H")) },{ FName(TEXT("ASTC_RGB")) } },
-			{ { FName(TEXT("BC7")) },{ FName(TEXT("ASTC_RGBAuto")) } },
-			{ { FName(TEXT("AutoDXT")) },{ FName(TEXT("ASTC_RGBAuto")) } },
+			{ { FName(TEXT("DXT1")) },		{ FName(TEXT("ASTC_RGB")) } },
+			{ { FName(TEXT("DXT5")) },		{ FName(TEXT("ASTC_RGBA")) } },
+			{ { FName(TEXT("DXT5n")) },		{ FName(TEXT("ASTC_NormalAG")) } },
+			{ { FName(TEXT("BC5")) },		{ FName(TEXT("ASTC_NormalRG")) } },
+			{ { FName(TEXT("BC6H")) },		{ FName(TEXT("ASTC_RGB")) } },
+			{ { FName(TEXT("BC7")) },		{ FName(TEXT("ASTC_RGBAuto")) } },
+			{ { FName(TEXT("AutoDXT")) },	{ FName(TEXT("ASTC_RGBAuto")) } },
 		};
 
 		GetDefaultTextureFormatNamePerLayer(OutFormats.AddDefaulted_GetRef(), this, Texture, EngineSettings, false, false, 1);
@@ -526,17 +457,17 @@ public:
 
 	virtual void GetAllTextureFormats(TArray<FName>& OutFormats) const override
 	{
-		// we remap some of the defaults (with PVRTC and ASTC formats)
+		// we remap some of the defaults (with ASTC formats)
 		static FName FormatRemap[][2] =
 		{
 			// Default format:				ASTC format:
-			{ { FName(TEXT("DXT1")) },{ FName(TEXT("ASTC_RGB")) } },
-			{ { FName(TEXT("DXT5")) },{ FName(TEXT("ASTC_RGBA")) } },
-			{ { FName(TEXT("DXT5n")) },{ FName(TEXT("ASTC_NormalAG")) } },
-			{ { FName(TEXT("BC5")) },{ FName(TEXT("ASTC_NormalRG")) } },
-			{ { FName(TEXT("BC6H")) },{ FName(TEXT("ASTC_RGB")) } },
-			{ { FName(TEXT("BC7")) },{ FName(TEXT("ASTC_RGBAuto")) } },
-			{ { FName(TEXT("AutoDXT")) },{ FName(TEXT("ASTC_RGBAuto")) } },
+			{ { FName(TEXT("DXT1")) },		{ FName(TEXT("ASTC_RGB")) } },
+			{ { FName(TEXT("DXT5")) },		{ FName(TEXT("ASTC_RGBA")) } },
+			{ { FName(TEXT("DXT5n")) },		{ FName(TEXT("ASTC_NormalAG")) } },
+			{ { FName(TEXT("BC5")) },		{ FName(TEXT("ASTC_NormalRG")) } },
+			{ { FName(TEXT("BC6H")) },		{ FName(TEXT("ASTC_RGB")) } },
+			{ { FName(TEXT("BC7")) },		{ FName(TEXT("ASTC_RGBAuto")) } },
+			{ { FName(TEXT("AutoDXT")) },	{ FName(TEXT("ASTC_RGBAuto")) } },
 		};
 
 		GetAllDefaultTextureFormats(this, OutFormats, false);
@@ -569,65 +500,6 @@ public:
 };
 
 
-
-class FAndroid_PVRTCTargetPlatform : public FAndroidTargetPlatform
-{
-public:
-	FAndroid_PVRTCTargetPlatform(bool bIsClient) : FAndroidTargetPlatform(bIsClient)
-	{
-		this->PlatformInfo = PlatformInfo::FindPlatformInfo("Android_PVRTC");
-	}
-
-	virtual FString GetAndroidVariantName() const override
-	{
-		return TEXT("PVRTC");
-	}
-
-	virtual FText DisplayName() const override
-	{
-		return LOCTEXT("Android_PVRTC", "Android (PVRTC)");
-	}
-
-	virtual bool SupportsCompressedNonPOT() const override
-	{
-		return false;
-	}
-
-	virtual bool SupportsTextureFormat(FName Format) const override
-	{
-		if (Format == AndroidTexFormat::NamePVRTC2 ||
-			Format == AndroidTexFormat::NamePVRTC4 ||
-			Format == AndroidTexFormat::NameAutoPVRTC)
-		{
-			return true;
-		}
-		return false;
-	}
-
-	virtual bool SupportsTextureFormatCategory(EAndroidTextureFormatCategory Category) const override
-	{
-		return Category == EAndroidTextureFormatCategory::PVRTC;
-	}
-
-	virtual bool SupportedByExtensionsString(const FString& ExtensionsString, const int GLESVersion) const override
-	{
-		return ExtensionsString.Contains(TEXT("GL_IMG_texture_compression_pvrtc"));
-	}
-
-	virtual FText GetVariantDisplayName() const override
-	{
-		return LOCTEXT("Android_PVRTC_ShortName", "PVRTC");
-	}
-
-	virtual float GetVariantPriority() const override
-	{
-		float Priority;
-		return (GConfig->GetFloat(TEXT("/Script/AndroidRuntimeSettings.AndroidRuntimeSettings"), TEXT("TextureFormatPriority_PVRTC"), Priority, GEngineIni) ?
-			Priority : 0.8f) * 10.0f + (IsClientOnly() ? 0.25f : 0.5f);
-	}
-};
-
-
 class FAndroid_ETC2TargetPlatform : public FAndroidTargetPlatform
 {
 public:
@@ -682,120 +554,6 @@ public:
 	}
 };
 
-class FAndroid_ETC1TargetPlatform : public FAndroidTargetPlatform
-{
-public:
-
-	FAndroid_ETC1TargetPlatform(bool bIsClient) : FAndroidTargetPlatform(bIsClient)
-	{
-		this->PlatformInfo = PlatformInfo::FindPlatformInfo("Android_ETC1");
-	}
-
-	virtual FText DisplayName() const override
-	{
-		return LOCTEXT("Android_ETC1", "Android (ETC1)");
-	}
-
-	virtual FString GetAndroidVariantName() const override
-	{
-		return TEXT("ETC1");
-	}
-
-	virtual bool SupportsTextureFormat(FName Format) const override
-	{
-		if (Format == AndroidTexFormat::NameETC1 ||
-			Format == AndroidTexFormat::NameAutoETC1)
-		{
-			return true;
-		}
-
-		return false;
-	}
-
-	virtual bool SupportsTextureFormatCategory(EAndroidTextureFormatCategory Category) const override
-	{
-		return Category == EAndroidTextureFormatCategory::ETC1;
-	}
-
-	// End FAndroidTargetPlatform overrides
-
-	virtual bool SupportedByExtensionsString(const FString& ExtensionsString, const int GLESVersion) const override
-	{
-		return ExtensionsString.Contains(TEXT("GL_OES_compressed_ETC1_RGB8_texture"));
-	}
-
-	virtual FText GetVariantDisplayName() const override
-	{
-		return LOCTEXT("Android_ETC1_ShortName", "ETC1");
-	}
-
-	virtual float GetVariantPriority() const override
-	{
-		float Priority;
-		return (GConfig->GetFloat(TEXT("/Script/AndroidRuntimeSettings.AndroidRuntimeSettings"), TEXT("TextureFormatPriority_ETC1"), Priority, GEngineIni) ?
-			Priority : 0.1f) * 10.0f + (IsClientOnly() ? 0.25f : 0.5f);
-	}
-};
-
-/**
-* Android cooking platform which cooks only ETC1a based textures.
-*/
-class FAndroid_ETC1aTargetPlatform : public FAndroidTargetPlatform
-{
-public:
-
-	FAndroid_ETC1aTargetPlatform(bool bIsClient) : FAndroidTargetPlatform(bIsClient)
-	{
-		this->PlatformInfo = PlatformInfo::FindPlatformInfo("Android_ETC1a");
-	}
-
-	virtual FText DisplayName() const override
-	{
-		return LOCTEXT("Android_ETC1a", "Android (ETCa1)");
-	}
-
-	virtual FString GetAndroidVariantName() const override
-	{
-		return TEXT("ETC1a");
-	}
-
-	virtual bool SupportsTextureFormat(FName Format) const override
-	{
-		if (Format == AndroidTexFormat::NameAutoETC1a)
-		{
-			return true;
-		}
-
-		return false;
-	}
-
-	virtual bool SupportsTextureFormatCategory(EAndroidTextureFormatCategory Category) const override
-	{
-		return Category == EAndroidTextureFormatCategory::ETC1a;
-	}
-
-	// End FAndroidTargetPlatform overrides
-
-	virtual bool SupportedByExtensionsString(const FString& ExtensionsString, const int GLESVersion) const override
-	{
-		return GLESVersion >= 0x30000;
-	}
-
-	virtual FText GetVariantDisplayName() const override
-	{
-		return LOCTEXT("Android_ETC1a_ShortName", "ETC1a");
-	}
-
-	virtual float GetVariantPriority() const override
-	{
-		float Priority;
-		return GConfig->GetFloat(TEXT("/Script/AndroidRuntimeSettings.AndroidRuntimeSettings"), TEXT("TextureFormatPriority_ETC1a"), Priority, GEngineIni) ? Priority : 1.0f;
-	}
-};
-
-
-
-
 class FAndroid_MultiTargetPlatform : public FAndroidTargetPlatform
 {
 	TArray<ITargetPlatform*> FormatTargetPlatforms;
diff --git a/Engine/Source/Developer/Android/AndroidTargetPlatform/Private/AndroidTargetPlatformModule.cpp b/Engine/Source/Developer/Android/AndroidTargetPlatform/Private/AndroidTargetPlatformModule.cpp
--- a/Engine/Source/Developer/Android/AndroidTargetPlatform/Private/AndroidTargetPlatformModule.cpp
+++ b/Engine/Source/Developer/Android/AndroidTargetPlatform/Private/AndroidTargetPlatformModule.cpp
@@ -47,12 +47,8 @@ public:
 				bool bIsClient = Type == 1;
 				SinglePlatforms.Add(new FAndroidTargetPlatform(bIsClient));
 				SinglePlatforms.Add(new FAndroid_ASTCTargetPlatform(bIsClient));
-				SinglePlatforms.Add(new FAndroid_ATCTargetPlatform(bIsClient));
 				SinglePlatforms.Add(new FAndroid_DXTTargetPlatform(bIsClient));
-				SinglePlatforms.Add(new FAndroid_ETC1TargetPlatform(bIsClient));
-				SinglePlatforms.Add(new FAndroid_ETC1aTargetPlatform(bIsClient));
 				SinglePlatforms.Add(new FAndroid_ETC2TargetPlatform(bIsClient));
-				SinglePlatforms.Add(new FAndroid_PVRTCTargetPlatform(bIsClient));
 
 				// thse are used in NotifyMultiSelectedFormatsChanged, so track in another array
 				MultiPlatforms.Add(new FAndroid_MultiTargetPlatform(bIsClient));
diff --git a/Engine/Source/Developer/Android/Android_ETC1aTargetPlatform/Android_ETC1aTargetPlatform.Build.cs b/Engine/Source/Developer/Android/Android_ETC1aTargetPlatform/Android_ETC1aTargetPlatform.Build.cs
deleted file mode 100644
diff --git a/Engine/Source/Developer/Android/Android_ETC1aTargetPlatform/Private/Android_ETC1aTargetPlatformModule.cpp b/Engine/Source/Developer/Android/Android_ETC1aTargetPlatform/Private/Android_ETC1aTargetPlatformModule.cpp
deleted file mode 100644
diff --git a/Engine/Source/Developer/DesktopPlatform/Public/PlatformInfo.h b/Engine/Source/Developer/DesktopPlatform/Public/PlatformInfo.h
--- a/Engine/Source/Developer/DesktopPlatform/Public/PlatformInfo.h
+++ b/Engine/Source/Developer/DesktopPlatform/Public/PlatformInfo.h
@@ -44,7 +44,7 @@ namespace PlatformInfo
 			/** The flavor generates different output when building (eg, 32 or 64-bit) */
 			BuildFlavor = 1<<0,
 
-			/** The flavor generates different output when cooking (eg, ATC or PVRTC texture format) */
+			/** The flavor generates different output when cooking (eg, ETC2 or ASTC texture format) */
 			CookFlavor = 1<<1,
 		};
 	}
@@ -102,16 +102,16 @@ namespace PlatformInfo
 	/** Information about a given platform */
 	struct FPlatformInfo
 	{
-		/** Name used to identify this platform, eg "Android_ATC" */
+		/** Name used to identify this platform, eg "Android_ETC2" */
 		FName PlatformInfoName;
 
 		/** Name used to find the corresponding ITargetPlatform for this platform (also used by UAT) */
 		FName TargetPlatformName;
 
-		/** Vanilla name for this platform, eg "Android" for "Android_ATC" */
+		/** Vanilla name for this platform, eg "Android" for "Android_ETC2" */
 		FName VanillaPlatformName;
 
-		/** Platform flavor, eg "ATC" for "Android_ATC" */
+		/** Platform flavor, eg "ETC2" for "Android_ETC2" */
 		FName PlatformFlavor;
 
 		/** The friendly (and localized) display name of this platform */
diff --git a/Engine/Source/Developer/Lumin/LuminTargetPlatform/Private/LuminTargetDevice.h b/Engine/Source/Developer/Lumin/LuminTargetPlatform/Private/LuminTargetDevice.h
--- a/Engine/Source/Developer/Lumin/LuminTargetPlatform/Private/LuminTargetDevice.h
+++ b/Engine/Source/Developer/Lumin/LuminTargetPlatform/Private/LuminTargetDevice.h
@@ -16,7 +16,7 @@ public:
 	*
 	* @param InTargetPlatform - The target platform.
 	* @param InSerialNumber - The ADB serial number of the target device.
-	* @param InAndroidVariant - The variant of the Android platform, i.e. ATC, DXT or PVRTC.
+	* @param InAndroidVariant - The variant of the Android platform, i.e. ETC2, DXT or ASTC.
 	*/
 	FLuminTargetDevice(const ITargetPlatform& InTargetPlatform, const FString& InSerialNumber, const FString& InAndroidVariant)
 		: FAndroidTargetDevice(InTargetPlatform, InSerialNumber, InAndroidVariant)
diff --git a/Engine/Source/Developer/TargetPlatform/Public/Interfaces/ITargetPlatform.h b/Engine/Source/Developer/TargetPlatform/Public/Interfaces/ITargetPlatform.h
--- a/Engine/Source/Developer/TargetPlatform/Public/Interfaces/ITargetPlatform.h
+++ b/Engine/Source/Developer/TargetPlatform/Public/Interfaces/ITargetPlatform.h
@@ -538,7 +538,7 @@ public:
 
 	/**
 	 * Gets the variant display name of this platform.
-	 * eg. For Android: "ETC1", "ETC2", ...
+	 * eg. For Android: "ASTC", "ETC2", ...
 	 *
 	 * @return FText display name for this platform variant.
 	 */
diff --git a/Engine/Source/Developer/TextureFormatAndroid/Private/TextureFormatAndroid.cpp b/Engine/Source/Developer/TextureFormatAndroid/Private/TextureFormatAndroid.cpp
--- a/Engine/Source/Developer/TextureFormatAndroid/Private/TextureFormatAndroid.cpp
+++ b/Engine/Source/Developer/TextureFormatAndroid/Private/TextureFormatAndroid.cpp
@@ -19,13 +19,6 @@ DEFINE_LOG_CATEGORY_STATIC(LogTextureFormatAndroid, Log, All);
  * Macro trickery for supported format names.
  */
 #define ENUM_SUPPORTED_FORMATS(op) \
-	op(ATC_RGB) \
-	op(ATC_RGBA_E) \
-	op(ATC_RGBA_I) \
-	op(AutoATC) \
-	op(ETC1) \
-	op(AutoETC1) \
-	op(AutoETC1a) \
 	op(ETC2_RGB) \
 	op(ETC2_RGBA) \
 	op(AutoETC2)
@@ -64,7 +57,7 @@ static bool CompressImageUsingQonvert(
 	// If block size changes, please update in AndroidETC.cpp in DecompressTexture
 	const int32 BlockSizeX = 4;
 	const int32 BlockSizeY = 4;
-	const int32 BlockBytes = (PixelFormat == PF_ATC_RGBA_E || PixelFormat == PF_ATC_RGBA_I || PixelFormat == PF_ETC2_RGBA) ? 16 : 8;
+	const int32 BlockBytes = (PixelFormat == PF_ETC2_RGBA) ? 16 : 8;
 	const int32 ImageBlocksX = FMath::Max(SizeX / BlockSizeX, 1);
 	const int32 ImageBlocksY = FMath::Max(SizeY / BlockSizeY, 1);
 
@@ -99,24 +92,12 @@ static bool CompressImageUsingQonvert(
 
 	switch (PixelFormat)
 	{
-	case PF_ETC1:
-		DstImg.nFormat = Q_FORMAT_ETC1_RGB8;
-		break;
 	case PF_ETC2_RGB:
 		DstImg.nFormat = Q_FORMAT_ETC2_RGB8;
 		break;
 	case PF_ETC2_RGBA:
 		DstImg.nFormat = Q_FORMAT_ETC2_RGBA8;
 		break;
-	case PF_ATC_RGB:
-		DstImg.nFormat = Q_FORMAT_ATC_RGB;
-		break;
-	case PF_ATC_RGBA_E:
-		DstImg.nFormat = Q_FORMAT_ATC_RGBA_EXPLICIT_ALPHA;
-		break;
-	case PF_ATC_RGBA_I:
-		DstImg.nFormat = Q_FORMAT_ATC_RGBA_INTERPOLATED_ALPHA;
-		break;
 	default:
 		UE_LOG(LogTextureFormatAndroid, Fatal, TEXT("Unsupported EPixelFormat for compression: %u"), (uint32)PixelFormat);
 		return false;
@@ -132,7 +113,7 @@ static bool CompressImageUsingQonvert(
 
 
 /**
- * ATITC and ETC1/2 texture format handler.
+ * ETC2 texture format handler.
  */
 class FTextureFormatAndroid : public ITextureFormat
 {
@@ -174,39 +155,6 @@ class FTextureFormatAndroid : public ITextureFormat
 
 		EPixelFormat CompressedPixelFormat = PF_Unknown;
 
-		if (BuildSettings.TextureFormatName == GTextureFormatNameAutoETC1)
-		{
-			if (bImageHasAlphaChannel)
-			{
-				// ETC1 can't support an alpha channel, store uncompressed
-				OutCompressedImage.SizeX = Image.SizeX;
-				OutCompressedImage.SizeY = Image.SizeY;
-				OutCompressedImage.SizeZ = (BuildSettings.bVolume || BuildSettings.bTextureArray) ? Image.NumSlices : 1;
-				OutCompressedImage.PixelFormat = PF_B8G8R8A8;
-				OutCompressedImage.RawData = Image.RawData;
-				return true;
-			}
-			else
-			{
-				CompressedPixelFormat = PF_ETC1;
-			}
-		}
-		else		
-		if (BuildSettings.TextureFormatName == GTextureFormatNameETC1)
-		{
-			CompressedPixelFormat = PF_ETC1;
-		}
-		else
-		if (BuildSettings.TextureFormatName == GTextureFormatNameAutoETC1a && !bImageHasAlphaChannel)
-		{
-			CompressedPixelFormat = PF_ETC1;
-		}
-		else
-		if (BuildSettings.TextureFormatName == GTextureFormatNameAutoETC1a && bImageHasAlphaChannel)
-		{
-			CompressedPixelFormat = PF_ETC2_RGBA;
-		}
-		else
 		if (BuildSettings.TextureFormatName == GTextureFormatNameETC2_RGB ||
 		   (BuildSettings.TextureFormatName == GTextureFormatNameAutoETC2 && !bImageHasAlphaChannel))
 		{
@@ -218,23 +166,6 @@ class FTextureFormatAndroid : public ITextureFormat
 		{
 			CompressedPixelFormat = PF_ETC2_RGBA;
 		}
-		else
-		if (BuildSettings.TextureFormatName == GTextureFormatNameATC_RGB || 
-		   (BuildSettings.TextureFormatName == GTextureFormatNameAutoATC && !bImageHasAlphaChannel))
-		{
-			CompressedPixelFormat = PF_ATC_RGB;
-		}
-		else
-		if (BuildSettings.TextureFormatName == GTextureFormatNameATC_RGBA_I ||
-		   (BuildSettings.TextureFormatName == GTextureFormatNameAutoATC && bImageHasAlphaChannel) )
-		{
-			CompressedPixelFormat = PF_ATC_RGBA_I;
-		}
-		else
-		if (BuildSettings.TextureFormatName == GTextureFormatNameATC_RGBA_E)
-		{
-			CompressedPixelFormat = PF_ATC_RGBA_E;
-		}
 
 		check(CompressedPixelFormat != PF_Unknown);
 
@@ -265,9 +196,6 @@ class FTextureFormatAndroid : public ITextureFormat
 	}
 };
 
-/**
- * Module for ATITC and ETC1/2 texture compression.
- */
 static ITextureFormat* Singleton = NULL;
 
 
diff --git a/Engine/Source/Editor/UnrealEd/Private/CookOnTheFlyServer.cpp b/Engine/Source/Editor/UnrealEd/Private/CookOnTheFlyServer.cpp
--- a/Engine/Source/Editor/UnrealEd/Private/CookOnTheFlyServer.cpp
+++ b/Engine/Source/Editor/UnrealEd/Private/CookOnTheFlyServer.cpp
@@ -7715,7 +7715,7 @@ void UCookOnTheFlyServer::StartCookByTheBook( const FCookByTheBookStartupOptions
 			{
 				using namespace PlatformInfo;
 				// Check all possible flavors 
-				// For example release version could be cooked as Android_ETC1 flavor, but DLC can be made as Android_ETC2
+				// For example release version could be cooked as Android_ASTC flavor, but DLC can be made as Android_ETC2
 				FVanillaPlatformEntry VanillaPlatfromEntry = BuildPlatformHierarchy(PlatformName, EPlatformFilter::CookFlavor);
 				for (const FPlatformInfo* PlatformFlaworInfo : VanillaPlatfromEntry.PlatformFlavors)
 				{
diff --git a/Engine/Source/Programs/AutomationTool/Android/AndroidPlatform.Automation.cs b/Engine/Source/Programs/AutomationTool/Android/AndroidPlatform.Automation.cs
--- a/Engine/Source/Programs/AutomationTool/Android/AndroidPlatform.Automation.cs
+++ b/Engine/Source/Programs/AutomationTool/Android/AndroidPlatform.Automation.cs
@@ -2371,19 +2371,6 @@ public class AndroidPlatformMulti : AndroidPlatform
     }
 }
 
-public class AndroidPlatformATC : AndroidPlatform
-{
-    public override string GetCookPlatform(bool bDedicatedServer, bool bIsClientOnly)
-    {
-		return bIsClientOnly ? "Android_ATCClient" : "Android_ATC";
-	}
-
-	public override TargetPlatformDescriptor GetTargetPlatformDescriptor()
-    {
-        return new TargetPlatformDescriptor(TargetPlatformType, "ATC");
-    }
-}
-
 public class AndroidPlatformDXT : AndroidPlatform
 {
     public override string GetCookPlatform(bool bDedicatedServer, bool bIsClientOnly)
@@ -2396,32 +2383,6 @@ public class AndroidPlatformDXT : AndroidPlatform
         return new TargetPlatformDescriptor(TargetPlatformType, "DXT");
     }
 }
-
-public class AndroidPlatformETC1 : AndroidPlatform
-{
-    public override string GetCookPlatform(bool bDedicatedServer, bool bIsClientOnly)
-    {
-		return bIsClientOnly ? "Android_ETC1Client" : "Android_ETC1";
-    }
-
-    public override TargetPlatformDescriptor GetTargetPlatformDescriptor()
-    {
-        return new TargetPlatformDescriptor(TargetPlatformType, "ETC1");
-    }
-}
-
-public class AndroidPlatformETC1a : AndroidPlatform
-{
-    public override string GetCookPlatform(bool bDedicatedServer, bool bIsClientOnly)
-    {
-        return "Android_ETC1a";
-    }
-    public override TargetPlatformDescriptor GetTargetPlatformDescriptor()
-    {
-        return new TargetPlatformDescriptor(TargetPlatformType, "ETC1a");
-    }
-}
-
 public class AndroidPlatformETC2 : AndroidPlatform
 {
     public override string GetCookPlatform(bool bDedicatedServer, bool bIsClientOnly)
@@ -2434,18 +2395,6 @@ public class AndroidPlatformETC2 : AndroidPlatform
     }
 }
 
-public class AndroidPlatformPVRTC : AndroidPlatform
-{
-    public override string GetCookPlatform(bool bDedicatedServer, bool bIsClientOnly)
-    {
-		return bIsClientOnly ? "Android_PVRTCClient" : "Android_PVRTC";
-    }
-    public override TargetPlatformDescriptor GetTargetPlatformDescriptor()
-    {
-        return new TargetPlatformDescriptor(TargetPlatformType, "PVRTC");
-    }
-}
-
 public class AndroidPlatformASTC : AndroidPlatform
 {
     public override string GetCookPlatform(bool bDedicatedServer, bool bIsClientOnly)
diff --git a/Engine/Source/Programs/UnrealBuildTool/Platform/Android/UEBuildAndroid.cs b/Engine/Source/Programs/UnrealBuildTool/Platform/Android/UEBuildAndroid.cs
--- a/Engine/Source/Programs/UnrealBuildTool/Platform/Android/UEBuildAndroid.cs
+++ b/Engine/Source/Programs/UnrealBuildTool/Platform/Android/UEBuildAndroid.cs
@@ -244,10 +244,9 @@ namespace UnrealBuildTool
 					else if (ModuleName == "TargetPlatform")
 					{
 						bBuildShaderFormats = true;
-						Rules.DynamicallyLoadedModuleNames.Add("TextureFormatPVR");
 						Rules.DynamicallyLoadedModuleNames.Add("TextureFormatDXT");
 						Rules.DynamicallyLoadedModuleNames.Add("TextureFormatASTC");
-						Rules.DynamicallyLoadedModuleNames.Add("TextureFormatAndroid");    // ATITC, ETC1 and ETC2
+						Rules.DynamicallyLoadedModuleNames.Add("TextureFormatAndroid");  // ETC2 
 						if (Target.bBuildDeveloperTools)
 						{
 							//Rules.DynamicallyLoadedModuleNames.Add("AudioFormatADPCM");	//@todo android: android audio
diff --git a/Engine/Source/Programs/UnrealBuildTool/Platform/Android/UEDeployAndroid.cs b/Engine/Source/Programs/UnrealBuildTool/Platform/Android/UEDeployAndroid.cs
--- a/Engine/Source/Programs/UnrealBuildTool/Platform/Android/UEDeployAndroid.cs
+++ b/Engine/Source/Programs/UnrealBuildTool/Platform/Android/UEDeployAndroid.cs
@@ -2151,12 +2151,12 @@ namespace UnrealBuildTool
 			}
 
 			//figure out which texture compressions are supported
-			bool bETC1Enabled, bETC1aEnabled, bETC2Enabled, bDXTEnabled, bATCEnabled, bPVRTCEnabled, bASTCEnabled;
-			bETC1Enabled = bETC1aEnabled = bETC2Enabled = bDXTEnabled = bATCEnabled = bPVRTCEnabled = bASTCEnabled = false;
+			bool bETC2Enabled, bDXTEnabled, bASTCEnabled;
+			bETC2Enabled = bDXTEnabled = bASTCEnabled = false;
 			if (CookFlavor.Length < 1)
 			{
 				//All values supported
-				bETC1Enabled = bETC2Enabled = bDXTEnabled = bATCEnabled = bPVRTCEnabled = bASTCEnabled = true;
+				bETC2Enabled = bDXTEnabled = bASTCEnabled = true;
 			}
 			else
 			{
@@ -2164,32 +2164,16 @@ namespace UnrealBuildTool
 				{
 					case "_Multi":
 						//need to check ini to determine which are supported
-						Ini.GetBool("/Script/AndroidRuntimeSettings.AndroidRuntimeSettings", "bMultiTargetFormat_ETC1", out bETC1Enabled);
-						Ini.GetBool("/Script/AndroidRuntimeSettings.AndroidRuntimeSettings", "bMultiTargetFormat_ETC1a", out bETC1aEnabled);
 						Ini.GetBool("/Script/AndroidRuntimeSettings.AndroidRuntimeSettings", "bMultiTargetFormat_ETC2", out bETC2Enabled);
 						Ini.GetBool("/Script/AndroidRuntimeSettings.AndroidRuntimeSettings", "bMultiTargetFormat_DXT", out bDXTEnabled);
-						Ini.GetBool("/Script/AndroidRuntimeSettings.AndroidRuntimeSettings", "bMultiTargetFormat_ATC", out bATCEnabled);
-						Ini.GetBool("/Script/AndroidRuntimeSettings.AndroidRuntimeSettings", "bMultiTargetFormat_PVRTC", out bPVRTCEnabled);
 						Ini.GetBool("/Script/AndroidRuntimeSettings.AndroidRuntimeSettings", "bMultiTargetFormat_ASTC", out bASTCEnabled);
 						break;
-					case "_ETC1":
-						bETC1Enabled = true;
-						break;
-					case "_ETC1a":
-						bETC1aEnabled = true;
-						break;
 					case "_ETC2":
 						bETC2Enabled = true;
 						break;
 					case "_DXT":
 						bDXTEnabled = true;
 						break;
-					case "_ATC":
-						bATCEnabled = true;
-						break;
-					case "_PVRTC":
-						bPVRTCEnabled = true;
-						break;
 					case "_ASTC":
 						bASTCEnabled = true;
 						break;
@@ -2198,17 +2182,13 @@ namespace UnrealBuildTool
 						break;
 				}
 			}
-			bool bSupportingAllTextureFormats = bETC1Enabled && bETC2Enabled && bDXTEnabled && bATCEnabled && bPVRTCEnabled && bASTCEnabled;
+			bool bSupportingAllTextureFormats = bETC2Enabled && bDXTEnabled && bASTCEnabled;
 
 			// If it is only ETC2 we need to skip adding the texture format filtering and instead use ES 3.0 as minimum version (it requires ETC2)
-			bool bOnlyETC2Enabled = (bETC2Enabled && !(bETC1aEnabled || bETC1Enabled || bDXTEnabled || bATCEnabled || bPVRTCEnabled || bASTCEnabled));
+			bool bOnlyETC2Enabled = (bETC2Enabled && !(bDXTEnabled || bASTCEnabled));
 
-			// Store cooked flavors in metadata (ETC1a treated as ETC1)
-			string CookedFlavors = ((bETC1Enabled || bETC1aEnabled) ? "ETC1," : "") +
-									(bETC2Enabled ? "ETC2," : "") +
+			string CookedFlavors = (bETC2Enabled ? "ETC2," : "") +
 									(bDXTEnabled ? "DXT," : "") +
-									(bATCEnabled ? "ATC," : "") +
-									(bPVRTCEnabled ? "PVRTC," : "") +
 									(bASTCEnabled ? "ASTC," : "");
 			CookedFlavors = (CookedFlavors == "") ? "" : CookedFlavors.Substring(0, CookedFlavors.Length - 1);
 
@@ -2501,30 +2481,17 @@ namespace UnrealBuildTool
 				if (!bSupportingAllTextureFormats)
 				{
 					Text.AppendLine("\t<!-- Supported texture compression formats (cooked) -->");
-					if (bETC1Enabled || bETC1aEnabled)
-					{
-						Text.AppendLine("\t<supports-gl-texture android:name=\"GL_OES_compressed_ETC1_RGB8_texture\" />");
-					}
 					if (bETC2Enabled && !bOnlyETC2Enabled)
 					{
 						Text.AppendLine("\t<supports-gl-texture android:name=\"GL_COMPRESSED_RGB8_ETC2\" />");
 						Text.AppendLine("\t<supports-gl-texture android:name=\"GL_COMPRESSED_RGBA8_ETC2_EAC\" />");
 					}
-					if (bATCEnabled)
-					{
-						Text.AppendLine("\t<supports-gl-texture android:name=\"GL_AMD_compressed_ATC_texture\" />");
-						Text.AppendLine("\t<supports-gl-texture android:name=\"GL_ATI_texture_compression_atitc\" />");
-					}
 					if (bDXTEnabled)
 					{
 						Text.AppendLine("\t<supports-gl-texture android:name=\"GL_EXT_texture_compression_dxt1\" />");
 						Text.AppendLine("\t<supports-gl-texture android:name=\"GL_EXT_texture_compression_s3tc\" />");
 						Text.AppendLine("\t<supports-gl-texture android:name=\"GL_NV_texture_compression_s3tc\" />");
 					}
-					if (bPVRTCEnabled)
-					{
-						Text.AppendLine("\t<supports-gl-texture android:name=\"GL_IMG_texture_compression_pvrtc\" />");
-					}
 					if (bASTCEnabled)
 					{
 						Text.AppendLine("\t<supports-gl-texture android:name=\"GL_KHR_texture_compression_astc_ldr\" />");
diff --git a/Engine/Source/Runtime/Android/AndroidRuntimeSettings/Classes/AndroidRuntimeSettings.h b/Engine/Source/Runtime/Android/AndroidRuntimeSettings/Classes/AndroidRuntimeSettings.h
--- a/Engine/Source/Runtime/Android/AndroidRuntimeSettings/Classes/AndroidRuntimeSettings.h
+++ b/Engine/Source/Runtime/Android/AndroidRuntimeSettings/Classes/AndroidRuntimeSettings.h
@@ -570,14 +570,6 @@ public:
 	UPROPERTY(GlobalConfig, EditAnywhere, Category = GraphicsDebugger)
 	FDirectoryPath MaliGraphicsDebuggerPath;
 
-	/** Include ETC1 textures when packaging with the Android (Multi) variant. ETC1 will be included if no other formats are selected. */
-	UPROPERTY(GlobalConfig, EditAnywhere, Category = MultiTextureFormats, meta = (DisplayName = "Include ETC1 textures"))
-	bool bMultiTargetFormat_ETC1;
-
-	/** Include ETC1a textures when packaging with the Android (Multi) variant. */
-	UPROPERTY(GlobalConfig, EditAnywhere, Category = MultiTextureFormats, meta = (DisplayName = "Include ETC1 - alpha textures"))
-	bool bMultiTargetFormat_ETC1a;
-
 	/** Include ETC2 textures when packaging with the Android (Multi) variant. */
 	UPROPERTY(GlobalConfig, EditAnywhere, Category = MultiTextureFormats, meta = (DisplayName = "Include ETC2 textures"))
 	bool bMultiTargetFormat_ETC2;
@@ -586,26 +578,10 @@ public:
 	UPROPERTY(GlobalConfig, EditAnywhere, Category = MultiTextureFormats, meta = (DisplayName = "Include DXT textures"))
 	bool bMultiTargetFormat_DXT;
 
-	/** Include PVRTC textures when packaging with the Android (Multi) variant. */
-	UPROPERTY(GlobalConfig, EditAnywhere, Category = MultiTextureFormats, meta = (DisplayName = "Include PVRTC textures"))
-	bool bMultiTargetFormat_PVRTC;
-
-	/** Include ATC textures when packaging with the Android (Multi) variant. */
-	UPROPERTY(GlobalConfig, EditAnywhere, Category = MultiTextureFormats, meta = (DisplayName = "Include ATC textures"))
-	bool bMultiTargetFormat_ATC;
-
 	/** Include ASTC textures when packaging with the Android (Multi) variant. */
 	UPROPERTY(GlobalConfig, EditAnywhere, Category = MultiTextureFormats, meta = (DisplayName = "Include ASTC textures"))
 	bool bMultiTargetFormat_ASTC;
 
-	/** Priority for the ETC1 texture format when launching on device or packaging using Android_Multi. The highest priority format supported by the device will be used. Default value is 0.1. */
-	UPROPERTY(GlobalConfig, EditAnywhere, Category = TextureFormatPriorities, meta = (DisplayName = "ETC1 texture format priority"))
-	float TextureFormatPriority_ETC1;
-
-	/** Priority for the ETC1a texture format when launching on device or packaging using Android_Multi. The highest priority format supported by the device will be used. Default value is 0.2. */
-	UPROPERTY(GlobalConfig, EditAnywhere, Category = TextureFormatPriorities, meta = (DisplayName = "ETC1 - alpha texture format priority"))
-	float TextureFormatPriority_ETC1a;
-
 	/** Priority for the ETC2 texture format when launching on device or packaging using Android_Multi. The highest priority format supported by the device will be used. Default value is 0.2. */
 	UPROPERTY(GlobalConfig, EditAnywhere, Category = TextureFormatPriorities, meta = (DisplayName = "ETC2 texture format priority"))
 	float TextureFormatPriority_ETC2;
@@ -614,14 +590,6 @@ public:
 	UPROPERTY(GlobalConfig, EditAnywhere, Category = TextureFormatPriorities, meta = (DisplayName = "DXT texture format priority"))
 	float TextureFormatPriority_DXT;
 
-	/** Priority for the PVRTC texture format when launching on device or packaging using Android_Multi. The highest priority format supported by the device will be used. Default value is 0.8. */
-	UPROPERTY(GlobalConfig, EditAnywhere, Category = TextureFormatPriorities, meta = (DisplayName = "PVRTC texture format priority"))
-	float TextureFormatPriority_PVRTC;
-
-	/** Priority for the ATC texture format when launching on device or packaging using Android_Multi. The highest priority format supported by the device will be used. Default value is 0.5. */
-	UPROPERTY(GlobalConfig, EditAnywhere, Category = TextureFormatPriorities, meta = (DisplayName = "ATC texture format priority"))
-	float TextureFormatPriority_ATC;
-
 	/** Priority for the ASTC texture format when launching on device or packaging using Android_Multi. The highest priority format supported by the device will be used. Default value is 0.9. */
 	UPROPERTY(GlobalConfig, EditAnywhere, Category = TextureFormatPriorities, meta = (DisplayName = "ASTC texture format priority"))
 	float TextureFormatPriority_ASTC;
diff --git a/Engine/Source/Runtime/Android/AndroidRuntimeSettings/Private/AndroidRuntimeSettings.cpp b/Engine/Source/Runtime/Android/AndroidRuntimeSettings/Private/AndroidRuntimeSettings.cpp
--- a/Engine/Source/Runtime/Android/AndroidRuntimeSettings/Private/AndroidRuntimeSettings.cpp
+++ b/Engine/Source/Runtime/Android/AndroidRuntimeSettings/Private/AndroidRuntimeSettings.cpp
@@ -28,19 +28,11 @@ UAndroidRuntimeSettings::UAndroidRuntimeSettings(const FObjectInitializer& Objec
 	, AudioSampleRate(44100)
 	, AudioCallbackBufferFrameSize(1024)
 	, AudioNumBuffersToEnqueue(4)
-	, bMultiTargetFormat_ETC1(true)
-	, bMultiTargetFormat_ETC1a(true)
 	, bMultiTargetFormat_ETC2(true)
 	, bMultiTargetFormat_DXT(true)
-	, bMultiTargetFormat_PVRTC(true)
-	, bMultiTargetFormat_ATC(true)
 	, bMultiTargetFormat_ASTC(true)
-	, TextureFormatPriority_ETC1(0.1f)
-	, TextureFormatPriority_ETC1a(0.18f)
 	, TextureFormatPriority_ETC2(0.2f)
 	, TextureFormatPriority_DXT(0.6f)
-	, TextureFormatPriority_PVRTC(0.8f)
-	, TextureFormatPriority_ATC(0.5f)
 	, TextureFormatPriority_ASTC(0.9f)
 {
 	bBuildForES31 = bBuildForES31 || !bSupportsVulkan;
@@ -119,10 +111,10 @@ void UAndroidRuntimeSettings::PostEditChangeProperty(struct FPropertyChangedEven
 		UpdateSinglePropertyInConfigFile(PropertyChangedEvent.Property, GetDefaultConfigFilename());
 
 		// Ensure we have at least one format for Android_Multi
-		if (!bMultiTargetFormat_ETC1 && !bMultiTargetFormat_ETC1a && !bMultiTargetFormat_ETC2 && !bMultiTargetFormat_DXT && !bMultiTargetFormat_PVRTC && !bMultiTargetFormat_ATC && !bMultiTargetFormat_ASTC)
+		if (!bMultiTargetFormat_ETC2 && !bMultiTargetFormat_DXT && !bMultiTargetFormat_ASTC)
 		{
-			bMultiTargetFormat_ETC1 = true;
-			UpdateSinglePropertyInConfigFile(GetClass()->FindPropertyByName(GET_MEMBER_NAME_CHECKED(UAndroidRuntimeSettings, bMultiTargetFormat_ETC1)), GetDefaultConfigFilename());
+			bMultiTargetFormat_ETC2 = true;
+			UpdateSinglePropertyInConfigFile(GetClass()->FindPropertyByName(GET_MEMBER_NAME_CHECKED(UAndroidRuntimeSettings, bMultiTargetFormat_ETC2)), GetDefaultConfigFilename());
 		}
 
 		// Notify the AndroidTargetPlatform module if it's loaded
diff --git a/Engine/Source/Runtime/Apple/MetalRHI/Private/MetalTexture.cpp b/Engine/Source/Runtime/Apple/MetalRHI/Private/MetalTexture.cpp
--- a/Engine/Source/Runtime/Apple/MetalRHI/Private/MetalTexture.cpp
+++ b/Engine/Source/Runtime/Apple/MetalRHI/Private/MetalTexture.cpp
@@ -133,10 +133,6 @@ static bool IsPixelFormatCompressed(EPixelFormat Format)
 		case PF_PVRTC4:
 		case PF_BC4:
 		case PF_BC5:
-		case PF_ATC_RGB:
-		case PF_ATC_RGBA_E:
-		case PF_ATC_RGBA_I:
-		case PF_ETC1:
 		case PF_ETC2_RGB:
 		case PF_ETC2_RGBA:
 		case PF_ASTC_4x4:
@@ -173,10 +169,6 @@ static bool IsPixelFormatPVRTCCompressed(EPixelFormat Format)
 	{
 		case PF_PVRTC2:
 		case PF_PVRTC4:
-		case PF_ATC_RGB:
-		case PF_ATC_RGBA_E:
-		case PF_ATC_RGBA_I:
-		case PF_ETC1:
 		case PF_ETC2_RGB:
 		case PF_ETC2_RGBA:
 			return true;
diff --git a/Engine/Source/Runtime/Core/Public/PixelFormat.h b/Engine/Source/Runtime/Core/Public/PixelFormat.h
--- a/Engine/Source/Runtime/Core/Public/PixelFormat.h
+++ b/Engine/Source/Runtime/Core/Public/PixelFormat.h
@@ -51,12 +51,12 @@ enum EPixelFormat
 	PF_R8G8B8A8             =37,
 	PF_A8R8G8B8				=38,	// Only used for legacy loading; do NOT use!
 	PF_BC4					=39,
-	PF_R8G8                 =40,
-	PF_ATC_RGB				=41,
-	PF_ATC_RGBA_E			=42,
-	PF_ATC_RGBA_I			=43,
+	PF_R8G8                 =40,	
+	PF_ATC_RGB				=41,	// Unsupported Format
+	PF_ATC_RGBA_E			=42,	// Unsupported Format
+	PF_ATC_RGBA_I			=43,	// Unsupported Format
 	PF_X24_G8				=44,	// Used for creating SRVs to alias a DepthStencil buffer to read Stencil. Don't use for creating textures.
-	PF_ETC1					=45,
+	PF_ETC1					=45,	// Unsupported Format
 	PF_ETC2_RGB				=46,
 	PF_ETC2_RGBA			=47,
 	PF_R32G32B32A32_UINT	=48,
diff --git a/Engine/Source/Runtime/Engine/Private/Texture2D.cpp b/Engine/Source/Runtime/Engine/Private/Texture2D.cpp
--- a/Engine/Source/Runtime/Engine/Private/Texture2D.cpp
+++ b/Engine/Source/Runtime/Engine/Private/Texture2D.cpp
@@ -1127,7 +1127,7 @@ bool UTexture2D::HasSameSourceArt(UTexture2D* InTexture)
 
 bool UTexture2D::HasAlphaChannel() const
 {
-	if (PlatformData && (PlatformData->PixelFormat != PF_DXT1) && (PlatformData->PixelFormat != PF_ATC_RGB))
+	if (PlatformData && (PlatformData->PixelFormat != PF_DXT1))
 	{
 		return true;
 	}
diff --git a/Engine/Source/Runtime/OpenGLDrv/Private/Android/AndroidOpenGL.cpp b/Engine/Source/Runtime/OpenGLDrv/Private/Android/AndroidOpenGL.cpp
--- a/Engine/Source/Runtime/OpenGLDrv/Private/Android/AndroidOpenGL.cpp
+++ b/Engine/Source/Runtime/OpenGLDrv/Private/Android/AndroidOpenGL.cpp
@@ -926,12 +926,6 @@ void FAndroidOpenGL::ProcessExtensions(const FString& ExtensionsString)
 			bool bFoundSupported = false;
 			for (FString Flavor : CookedFlavors)
 			{
-				if (Flavor.Equals(TEXT("ETC1")))
-				{
-					// every device supports ETC1
-					bFoundSupported = true;
-					break;
-				}
 				if (Flavor.Equals(TEXT("ETC2")))
 				{
 					if (FOpenGL::SupportsETC2())
@@ -940,14 +934,6 @@ void FAndroidOpenGL::ProcessExtensions(const FString& ExtensionsString)
 						break;
 					}
 				}
-				if (Flavor.Equals(TEXT("ATC")))
-				{
-					if (FOpenGL::SupportsATITC())
-					{
-						bFoundSupported = true;
-						break;
-					}
-				}
 				if (Flavor.Equals(TEXT("DXT")))
 				{
 					if (FOpenGL::SupportsDXT())
@@ -956,14 +942,6 @@ void FAndroidOpenGL::ProcessExtensions(const FString& ExtensionsString)
 						break;
 					}
 				}
-				if (Flavor.Equals(TEXT("PVRTC")))
-				{
-					if (FOpenGL::SupportsPVRTC())
-					{
-						bFoundSupported = true;
-						break;
-					}
-				}
 				if (Flavor.Equals(TEXT("ASTC")))
 				{
 					if (FOpenGL::SupportsASTC())
@@ -976,11 +954,8 @@ void FAndroidOpenGL::ProcessExtensions(const FString& ExtensionsString)
 
 			if (!bFoundSupported)
 			{
-				FString Message = TEXT("Cooked Flavors: ") + CookedFlavorsString + TEXT("\n\nSupported: ETC1") +
-					(FOpenGL::SupportsETC2() ? TEXT(",ETC2") : TEXT("")) +
-					(FOpenGL::SupportsATITC() ? TEXT(",ATC") : TEXT("")) +
+				FString Message = TEXT("Cooked Flavors: ") + CookedFlavorsString + TEXT("\n\nSupported: ETC2") +
 					(FOpenGL::SupportsDXT() ? TEXT(",DXT") : TEXT("")) +
-					(FOpenGL::SupportsPVRTC() ? TEXT(",PVRTC") : TEXT("")) +
 					(FOpenGL::SupportsASTC() ? TEXT(",ASTC") : TEXT(""));
 
 				FPlatformMisc::LowLevelOutputDebugStringf(TEXT("Error: Unsupported Texture Format\n%s"), *Message);
diff --git a/Engine/Source/Runtime/OpenGLDrv/Private/Android/AndroidOpenGLPrivate.h b/Engine/Source/Runtime/OpenGLDrv/Private/Android/AndroidOpenGLPrivate.h
--- a/Engine/Source/Runtime/OpenGLDrv/Private/Android/AndroidOpenGLPrivate.h
+++ b/Engine/Source/Runtime/OpenGLDrv/Private/Android/AndroidOpenGLPrivate.h
@@ -74,21 +74,9 @@ private:
 		{
 			TargetPlatformNames.Add(TEXT("Android_DXT"));
 		}
-		if (ExtensionsString.Contains(TEXT("GL_ATI_texture_compression_atitc")) || ExtensionsString.Contains(TEXT("GL_AMD_compressed_ATC_texture")))
-		{
-			TargetPlatformNames.Add(TEXT("Android_ATC"));
-		}
-		if (ExtensionsString.Contains(TEXT("GL_IMG_texture_compression_pvrtc")))
-		{
-			TargetPlatformNames.Add(TEXT("Android_PVRTC"));
-		}
 		
 		TargetPlatformNames.Add(TEXT("Android_ETC2"));
 
-		// all devices support ETC
-		TargetPlatformNames.Add(TEXT("Android_ETC1a"));
-		TargetPlatformNames.Add(TEXT("Android_ETC1"));
-
 		// finally, generic Android
 		TargetPlatformNames.Add(TEXT("Android"));
 
diff --git a/Engine/Source/Runtime/OpenGLDrv/Private/OpenGLDevice.cpp b/Engine/Source/Runtime/OpenGLDrv/Private/OpenGLDevice.cpp
--- a/Engine/Source/Runtime/OpenGLDrv/Private/OpenGLDevice.cpp
+++ b/Engine/Source/Runtime/OpenGLDrv/Private/OpenGLDevice.cpp
@@ -1093,21 +1093,6 @@ static void InitRHICapabilitiesForGL()
 			SetupTextureFormat( PF_DXT5,	FOpenGLTextureFormat(GL_COMPRESSED_RGBA_S3TC_DXT5_EXT,	GL_COMPRESSED_RGBA_S3TC_DXT5_EXT,		GL_RGBA,	GL_UNSIGNED_BYTE,	true,	false));
 		}
 	}
-	if ( FOpenGL::SupportsPVRTC() )
-	{
-		SetupTextureFormat( PF_PVRTC2,		FOpenGLTextureFormat(GL_COMPRESSED_RGBA_PVRTC_2BPPV1_IMG, GL_COMPRESSED_RGBA_PVRTC_2BPPV1_IMG,	GL_RGBA,	GL_UNSIGNED_BYTE,	true,	false));
-		SetupTextureFormat( PF_PVRTC4,		FOpenGLTextureFormat(GL_COMPRESSED_RGBA_PVRTC_4BPPV1_IMG, GL_COMPRESSED_RGBA_PVRTC_4BPPV1_IMG,	GL_RGBA,	GL_UNSIGNED_BYTE,	true,	false));
-	}
-	if ( FOpenGL::SupportsATITC() )
-	{
-		SetupTextureFormat( PF_ATC_RGB,		FOpenGLTextureFormat(GL_ATC_RGB_AMD,					 GL_ATC_RGB_AMD,						GL_RGBA,	GL_UNSIGNED_BYTE,	true,	false));
-		SetupTextureFormat( PF_ATC_RGBA_E,	FOpenGLTextureFormat(GL_ATC_RGBA_EXPLICIT_ALPHA_AMD,	 GL_ATC_RGBA_EXPLICIT_ALPHA_AMD,		GL_RGBA,	GL_UNSIGNED_BYTE,	true,	false));
-		SetupTextureFormat( PF_ATC_RGBA_I,	FOpenGLTextureFormat(GL_ATC_RGBA_INTERPOLATED_ALPHA_AMD, GL_ATC_RGBA_INTERPOLATED_ALPHA_AMD,	GL_RGBA,	GL_UNSIGNED_BYTE,	true,	false));
-	}
-	if ( FOpenGL::SupportsETC1() )
-	{
-		SetupTextureFormat( PF_ETC1,		FOpenGLTextureFormat(GL_ETC1_RGB8_OES,					GL_ETC1_RGB8_OES,						GL_RGBA,	GL_UNSIGNED_BYTE,	true,	false));
-	}
 #if PLATFORM_ANDROID && !PLATFORM_LUMINGL4
 	if ( FOpenGL::SupportsETC2() )
 	{
@@ -1115,12 +1100,6 @@ static void InitRHICapabilitiesForGL()
 		SetupTextureFormat( PF_ETC2_RGBA,		FOpenGLTextureFormat(GL_COMPRESSED_RGBA8_ETC2_EAC,	GL_COMPRESSED_SRGB8_ALPHA8_ETC2_EAC,		GL_RGBA,		GL_UNSIGNED_BYTE,	true,		false));
 		SetupTextureFormat( PF_ETC2_R11_EAC,	FOpenGLTextureFormat(GL_COMPRESSED_R11_EAC,			GL_COMPRESSED_R11_EAC,						GL_RED,			GL_UNSIGNED_BYTE,	true,		false));
 		SetupTextureFormat( PF_ETC2_RG11_EAC,	FOpenGLTextureFormat(GL_COMPRESSED_RG11_EAC,		GL_COMPRESSED_RG11_EAC,						GL_RG,			GL_UNSIGNED_BYTE,	true,		false));
-
-		// ETC2 is a superset of ETC1 with sRGB support
-		if (FOpenGL::SupportsSRGB())
-		{
-			SetupTextureFormat( PF_ETC1,	FOpenGLTextureFormat(GL_COMPRESSED_RGB8_ETC2, GL_COMPRESSED_SRGB8_ETC2,	GL_RGBA, GL_UNSIGNED_BYTE, true, false));
-		}
 	}
 #endif
 	if (FOpenGL::SupportsASTC())
diff --git a/Engine/Source/Runtime/OpenGLDrv/Private/OpenGLES.cpp b/Engine/Source/Runtime/OpenGLDrv/Private/OpenGLES.cpp
--- a/Engine/Source/Runtime/OpenGLDrv/Private/OpenGLES.cpp
+++ b/Engine/Source/Runtime/OpenGLDrv/Private/OpenGLES.cpp
@@ -185,15 +185,6 @@ bool FOpenGLES::bSupportsMultisampledRenderToTexture = false;
 /** GL_NV_texture_compression_s3tc, GL_EXT_texture_compression_s3tc */
 bool FOpenGLES::bSupportsDXT = false;
 
-/** GL_IMG_texture_compression_pvrtc */
-bool FOpenGLES::bSupportsPVRTC = false;
-
-/** GL_ATI_texture_compression_atitc, GL_AMD_compressed_ATC_texture */
-bool FOpenGLES::bSupportsATITC = false;
-
-/** GL_OES_compressed_ETC1_RGB8_texture */
-bool FOpenGLES::bSupportsETC1 = false;
-
 /** OpenGL ES 3.0 profile */
 bool FOpenGLES::bSupportsETC2 = false;
 
@@ -353,9 +344,6 @@ void FOpenGLES::ProcessExtensions(const FString& ExtensionsString)
 	bSupportsShaderDepthStencilFetch = ExtensionsString.Contains(TEXT("GL_ARM_shader_framebuffer_fetch_depth_stencil"));
 	bSupportsMultisampledRenderToTexture = ExtensionsString.Contains(TEXT("GL_EXT_multisampled_render_to_texture"));
 	bSupportsDXT = ExtensionsString.Contains(TEXT("GL_NV_texture_compression_s3tc")) || ExtensionsString.Contains(TEXT("GL_EXT_texture_compression_s3tc"));
-	bSupportsPVRTC = ExtensionsString.Contains(TEXT("GL_IMG_texture_compression_pvrtc"));
-	bSupportsATITC = ExtensionsString.Contains(TEXT("GL_ATI_texture_compression_atitc")) || ExtensionsString.Contains(TEXT("GL_AMD_compressed_ATC_texture"));
-	bSupportsETC1 = ExtensionsString.Contains(TEXT("GL_OES_compressed_ETC1_RGB8_texture"));
 	bSupportsVertexArrayObjects = ExtensionsString.Contains(TEXT("GL_OES_vertex_array_object"));
 	bSupportsDiscardFrameBuffer = ExtensionsString.Contains(TEXT("GL_EXT_discard_framebuffer"));
 	bSupportsNVFrameBufferBlit = ExtensionsString.Contains(TEXT("GL_NV_framebuffer_blit"));
@@ -578,13 +566,6 @@ void FOpenGLES::ProcessExtensions(const FString& ExtensionsString)
 		bSupportsColorBufferHalfFloat = false;
 	}
 
-	if (IsES31Usable())
-	{
-		// ES3.1 requires sRGB texture sampling, these formats do not support it
-		bSupportsATITC = false;
-		bSupportsPVRTC = false;
-	}
-
 	// Set lowest possible limits for texture units, to avoid extra work in GL RHI
 	MaxTextureImageUnits = FMath::Min(MaxTextureImageUnits, 16);
 	MaxVertexTextureImageUnits = FMath::Min(MaxVertexTextureImageUnits, 16);
diff --git a/Engine/Source/Runtime/OpenGLDrv/Private/OpenGLTexture.cpp b/Engine/Source/Runtime/OpenGLDrv/Private/OpenGLTexture.cpp
--- a/Engine/Source/Runtime/OpenGLDrv/Private/OpenGLTexture.cpp
+++ b/Engine/Source/Runtime/OpenGLDrv/Private/OpenGLTexture.cpp
@@ -724,12 +724,6 @@ void TOpenGLTexture<RHIResourceType>::Resolve(uint32 MipIndex,uint32 ArrayIndex)
 	const uint32 MipSizeY = FMath::Max(this->GetSizeY() >> MipIndex,BlockSizeY);
 	uint32 NumBlocksX = (MipSizeX + BlockSizeX - 1) / BlockSizeX;
 	uint32 NumBlocksY = (MipSizeY + BlockSizeY - 1) / BlockSizeY;
-	if ( PixelFormat == PF_PVRTC2 || PixelFormat == PF_PVRTC4 )
-	{
-		// PVRTC has minimum 2 blocks width and height
-		NumBlocksX = FMath::Max<uint32>(NumBlocksX, 2);
-		NumBlocksY = FMath::Max<uint32>(NumBlocksY, 2);
-	}
 	const uint32 MipBytes = NumBlocksX * NumBlocksY * BlockBytes;
 	
 	const int32 BufferIndex = MipIndex * (bCubemap ? 6 : 1) * this->GetEffectiveSizeZ() + ArrayIndex;
@@ -807,12 +801,6 @@ uint32 TOpenGLTexture<RHIResourceType>::GetLockSize(uint32 InMipIndex, uint32 Ar
 	const uint32 MipSizeY = FMath::Max(this->GetSizeY() >> InMipIndex, BlockSizeY);
 	uint32 NumBlocksX = (MipSizeX + BlockSizeX - 1) / BlockSizeX;
 	uint32 NumBlocksY = (MipSizeY + BlockSizeY - 1) / BlockSizeY;
-	if (PixelFormat == PF_PVRTC2 || PixelFormat == PF_PVRTC4)
-	{
-		// PVRTC has minimum 2 blocks width and height
-		NumBlocksX = FMath::Max<uint32>(NumBlocksX, 2);
-		NumBlocksY = FMath::Max<uint32>(NumBlocksY, 2);
-	}
 	const uint32 MipBytes = NumBlocksX * NumBlocksY * BlockBytes;
 	DestStride = NumBlocksX * BlockBytes;
 	return MipBytes;
@@ -1176,24 +1164,6 @@ void TOpenGLTexture<RHIResourceType>::Unlock(uint32 MipIndex,uint32 ArrayIndex)
 		
 		bool bIsCompressed = GLFormat.bCompressed;
 		GLint internalFormat = GLFormat.InternalFormat[bSRGB];
-#if PLATFORM_ANDROID
-		uint8* DecompressedPointer = nullptr;
-		if (bIsCompressed)
-		{
-			if (!FOpenGL::SupportsETC2() && this->GetFormat() == PF_ETC2_RGBA)
-			{
-				bIsCompressed = false;
-				internalFormat = GL_RGBA;
-
-				DecompressTexture((uint8_t *)LockedBuffer, FMath::Max<uint32>(1, (this->GetSizeX() >> MipIndex)), FMath::Max<uint32>(1, (this->GetSizeY() >> MipIndex)), GLFormat.InternalFormat[bSRGB], &DecompressedPointer);
-				if (!DecompressedPointer)
-				{
-					UE_LOG(LogRHI, Fatal, TEXT("ETC2 texture compression failed for fallback on ETC1 device."));
-				}
-				LockedBuffer = DecompressedPointer;
-			}
-		}
-#endif // PLATFORM_ANDROID
 
 		if (bIsCompressed)
 		{
@@ -1268,12 +1238,6 @@ void TOpenGLTexture<RHIResourceType>::Unlock(uint32 MipIndex,uint32 ArrayIndex)
 			// Unlock "PixelBuffer" and free the temp memory after the texture upload.
 			PixelBuffer->Unlock();
 		}
-#if PLATFORM_ANDROID
-		if (DecompressedPointer)
-		{
-			free(DecompressedPointer);
-		}
-#endif // PLATFORM_ANDROID
 	}
 
 	// No need to restore texture stage; leave it like this,
@@ -1354,12 +1318,6 @@ void TOpenGLTexture<RHIResourceType>::CloneViaPBO( TOpenGLTexture<RHIResourceTyp
 			const uint32 DataSizeY = FMath::Max(MipSizeY,BlockSizeY);
 			uint32 NumBlocksX = (DataSizeX + BlockSizeX - 1) / BlockSizeX;
 			uint32 NumBlocksY = (DataSizeY + BlockSizeY - 1) / BlockSizeY;
-			if ( PixelFormat == PF_PVRTC2 || PixelFormat == PF_PVRTC4 )
-			{
-				// PVRTC has minimum 2 blocks width and height
-				NumBlocksX = FMath::Max<uint32>(NumBlocksX, 2);
-				NumBlocksY = FMath::Max<uint32>(NumBlocksY, 2);
-			}
 
 			const uint32 MipBytes = NumBlocksX * NumBlocksY * BlockBytes;
 			const int32 BufferIndex = DstMipIndex * (bCubemap ? 6 : 1) * this->GetEffectiveSizeZ() + ArrayIndex;
diff --git a/Engine/Source/Runtime/OpenGLDrv/Public/OpenGL.h b/Engine/Source/Runtime/OpenGLDrv/Public/OpenGL.h
--- a/Engine/Source/Runtime/OpenGLDrv/Public/OpenGL.h
+++ b/Engine/Source/Runtime/OpenGLDrv/Public/OpenGL.h
@@ -138,10 +138,7 @@ public:
 	static FORCEINLINE bool SupportsSRGB()								{ return true; }
 	static FORCEINLINE bool SupportsRGBA8()								{ return true; }
 	static FORCEINLINE bool SupportsDXT()								{ return true; }
-	static FORCEINLINE bool SupportsPVRTC()								{ return false; }
-	static FORCEINLINE bool SupportsATITC()								{ return false; }
 	static FORCEINLINE bool SupportsASTC()								{ return bSupportsASTC; }
-	static FORCEINLINE bool SupportsETC1()								{ return false; }
 	static FORCEINLINE bool SupportsETC2()								{ return false; }
 	static FORCEINLINE bool SupportsFramebufferSRGBEnable()				{ return true; }
 	static FORCEINLINE bool SupportsCombinedDepthStencilAttachment()	{ return true; }
@@ -475,25 +472,6 @@ protected:
 #define GL_COMPRESSED_RGBA_S3TC_DXT5_EXT		0x83F3
 #endif
 
-// http://www.khronos.org/registry/gles/extensions/IMG/IMG_texture_compression_pvrtc.txt
-#if !defined(GL_IMG_texture_compression_pvrtc)
-#define GL_COMPRESSED_RGB_PVRTC_4BPPV1_IMG		0x8C00
-#define GL_COMPRESSED_RGB_PVRTC_2BPPV1_IMG		0x8C01
-#define GL_COMPRESSED_RGBA_PVRTC_4BPPV1_IMG		0x8C02
-#define GL_COMPRESSED_RGBA_PVRTC_2BPPV1_IMG		0x8C03
-#endif
-
-// http://www.khronos.org/registry/gles/extensions/AMD/AMD_compressed_ATC_texture.txt
-#if !defined(GL_ATI_texture_compression_atitc)
-#define GL_ATC_RGB_AMD							0x8C92
-#define GL_ATC_RGBA_EXPLICIT_ALPHA_AMD			0x8C93
-#define GL_ATC_RGBA_INTERPOLATED_ALPHA_AMD		0x87EE
-#endif
-
-#if !defined(GL_OES_compressed_ETC1_RGB8_texture)
-#define GL_ETC1_RGB8_OES                        0x8D64
-#endif
-
 // http://www.opengl.org/registry/specs/EXT/texture_sRGB.txt
 #if !defined(GL_EXT_texture_sRGB)
 #define GL_SRGB_EXT                       0x8C40
@@ -522,13 +500,6 @@ protected:
 #define GL_COMPRESSED_SIGNED_RG_RGTC2     0x8DBE
 #endif
 
-// http://www.khronos.org/registry/gles/extensions/AMD/AMD_compressed_ATC_texture.txt
-#if !defined(GL_ATI_texture_compression_atitc)
-#define GL_ATC_RGB_AMD							0x8C92
-#define GL_ATC_RGBA_EXPLICIT_ALPHA_AMD			0x8C93
-#define GL_ATC_RGBA_INTERPOLATED_ALPHA_AMD		0x87EE
-#endif
-
 /* http://www.khronos.org/registry/gles/extensions/NV/NV_sRGB_formats.txt */
 #if !defined(GL_NV_sRGB_formats)
 #define GL_SLUMINANCE_NV                                        0x8C46
@@ -540,7 +511,6 @@ protected:
 #define GL_COMPRESSED_SRGB_ALPHA_S3TC_DXT1_NV                   0x8C4D
 #define GL_COMPRESSED_SRGB_ALPHA_S3TC_DXT3_NV                   0x8C4E
 #define GL_COMPRESSED_SRGB_ALPHA_S3TC_DXT5_NV                   0x8C4F
-#define GL_ETC1_SRGB8_NV                                        0x88EE
 #endif
 
 // http://www.opengl.org/registry/specs/KHR/texture_compression_astc_ldr.txt
diff --git a/Engine/Source/Runtime/OpenGLDrv/Public/OpenGLES.h b/Engine/Source/Runtime/OpenGLDrv/Public/OpenGLES.h
--- a/Engine/Source/Runtime/OpenGLDrv/Public/OpenGLES.h
+++ b/Engine/Source/Runtime/OpenGLDrv/Public/OpenGLES.h
@@ -672,9 +672,6 @@ struct FOpenGLES : public FOpenGLBase
 	static FORCEINLINE bool SupportsSRGB() { return IsES31Usable(); }
 	static FORCEINLINE bool SupportsRGBA8() { return bSupportsRGBA8; }
 	static FORCEINLINE bool SupportsDXT() { return bSupportsDXT; }
-	static FORCEINLINE bool SupportsPVRTC() { return bSupportsPVRTC; }
-	static FORCEINLINE bool SupportsATITC() { return bSupportsATITC; }
-	static FORCEINLINE bool SupportsETC1() { return bSupportsETC1; }
 	static FORCEINLINE bool SupportsETC2() { return bSupportsETC2; }
 	static FORCEINLINE bool SupportsCombinedDepthStencilAttachment() { return false; }
 	static FORCEINLINE bool SupportsPackedDepthStencil() { return bSupportsPackedDepthStencil; }
@@ -1272,15 +1269,6 @@ protected:
 	/** GL_NV_texture_compression_s3tc, GL_EXT_texture_compression_s3tc */
 	static bool bSupportsDXT;
 
-	/** GL_IMG_texture_compression_pvrtc */
-	static bool bSupportsPVRTC;
-
-	/** GL_ATI_texture_compression_atitc, GL_AMD_compressed_ATC_texture */
-	static bool bSupportsATITC;
-
-	/** GL_OES_compressed_ETC1_RGB8_texture */
-	static bool bSupportsETC1;
-
 	/** OpenGL ES 3.0 profile */
 	static bool bSupportsETC2;
 
diff --git a/Engine/Source/Runtime/VulkanRHI/Private/VulkanDevice.cpp b/Engine/Source/Runtime/VulkanRHI/Private/VulkanDevice.cpp
--- a/Engine/Source/Runtime/VulkanRHI/Private/VulkanDevice.cpp
+++ b/Engine/Source/Runtime/VulkanRHI/Private/VulkanDevice.cpp
@@ -648,13 +648,6 @@ void FVulkanDevice::SetupFormats()
 			SetComponentMapping(PF_ASTC_12x12, VK_COMPONENT_SWIZZLE_R, VK_COMPONENT_SWIZZLE_G, VK_COMPONENT_SWIZZLE_B, VK_COMPONENT_SWIZZLE_A);
 		}
 
-		// ETC1 is a subset of ETC2 R8G8B8.
-		MapFormatSupport(PF_ETC1, VK_FORMAT_ETC2_R8G8B8_UNORM_BLOCK);
-		if (GPixelFormats[PF_ETC1].Supported)
-		{
-			SetComponentMapping(PF_ETC1, VK_COMPONENT_SWIZZLE_R, VK_COMPONENT_SWIZZLE_G, VK_COMPONENT_SWIZZLE_B, VK_COMPONENT_SWIZZLE_ONE);
-		}
-
 		MapFormatSupport(PF_ETC2_RGB, VK_FORMAT_ETC2_R8G8B8_UNORM_BLOCK);
 		if (GPixelFormats[PF_ETC2_RGB].Supported)
 		{
