[Android] Fix editor compile error. SupportsES2 has been removed

diff --git a/Engine/Source/Developer/Android/AndroidTargetPlatform/Private/AndroidTargetPlatform.cpp b/Engine/Source/Developer/Android/AndroidTargetPlatform/Private/AndroidTargetPlatform.cpp
--- a/Engine/Source/Developer/Android/AndroidTargetPlatform/Private/AndroidTargetPlatform.cpp
+++ b/Engine/Source/Developer/Android/AndroidTargetPlatform/Private/AndroidTargetPlatform.cpp
@@ -375,7 +375,7 @@ bool FAndroidTargetPlatform::SupportsFeature( ETargetPlatformFeatures Feature )
 			return UsesVirtualTextures();
 
 		case ETargetPlatformFeatures::LandscapeMeshLODStreaming:
-			return SupportsLandscapeMeshLODStreaming() && !SupportsES2();
+			return SupportsLandscapeMeshLODStreaming();
 
 		default:
 			break;
