UDN suggested fix for Android where BC7 compression was mapping to an uncompressed texture format instead of ASTC

diff --git a/Engine/Source/Developer/Android/AndroidTargetPlatform/Private/AndroidTargetPlatform.h b/Engine/Source/Developer/Android/AndroidTargetPlatform/Private/AndroidTargetPlatform.h
--- a/Engine/Source/Developer/Android/AndroidTargetPlatform/Private/AndroidTargetPlatform.h
+++ b/Engine/Source/Developer/Android/AndroidTargetPlatform/Private/AndroidTargetPlatform.h
@@ -430,7 +430,7 @@ public:
 			{ { FName(TEXT("AutoDXT")) },	{ FName(TEXT("ASTC_RGBAuto")) } },
 		};
 
-		GetDefaultTextureFormatNamePerLayer(OutFormats.AddDefaulted_GetRef(), this, Texture, EngineSettings, false, false, 1);
+		GetDefaultTextureFormatNamePerLayer(OutFormats.AddDefaulted_GetRef(), this, Texture, EngineSettings, true, false, 1);
 
 		for (FName& TextureFormatName : OutFormats.Last())
 		{
@@ -471,7 +471,7 @@ public:
 			{ { FName(TEXT("AutoDXT")) },	{ FName(TEXT("ASTC_RGBAuto")) } },
 		};
 
-		GetAllDefaultTextureFormats(this, OutFormats, false);
+		GetAllDefaultTextureFormats(this, OutFormats, true);
 
 		for (int32 RemapIndex = 0; RemapIndex < UE_ARRAY_COUNT(FormatRemap); ++RemapIndex)
 		{
