Support compressed volume textures on iOS and Android (ASTC)
Removed bSupportsVolumeTextureCompression as all platforms support it now, only Switch fallbacks to uncompressed

diff --git a/Engine/Config/Linux/DataDrivenPlatformInfo.ini b/Engine/Config/Linux/DataDrivenPlatformInfo.ini
--- a/Engine/Config/Linux/DataDrivenPlatformInfo.ini
+++ b/Engine/Config/Linux/DataDrivenPlatformInfo.ini
@@ -29,7 +29,6 @@ bSupportsMobileMultiView = false
 ; //#todo-vulkanddpi
 bSupportsArrayTextureCompression = false
 
-bSupportsVolumeTextureCompression = true
 bSupportsDistanceFields = true
 bSupportsDiaphragmDOF = true
 bSupportsRGBColorBuffer = true
diff --git a/Engine/Config/Windows/DataDrivenPlatformInfo.ini b/Engine/Config/Windows/DataDrivenPlatformInfo.ini
--- a/Engine/Config/Windows/DataDrivenPlatformInfo.ini
+++ b/Engine/Config/Windows/DataDrivenPlatformInfo.ini
@@ -26,7 +26,6 @@ UBTPlatformName=Win64
 Language=D3D
 MaxFeatureLevel=SM5
 bIsPC=true
-bSupportsVolumeTextureCompression=true
 bSupportsRayTracing=true
 bSupportsRayTracingIndirectInstanceData=true
 bSupportsWaveOperations=false
@@ -44,7 +43,6 @@ bSupportsAnisotropicMaterials=true
 Language=D3D
 MaxFeatureLevel=SM6
 bIsPC=true
-bSupportsVolumeTextureCompression=true
 bSupportsRayTracing=true
 bSupportsRayTracingIndirectInstanceData=true
 bSupportsGPUScene=true
@@ -70,7 +68,6 @@ bSupportsMobileMultiView = false
 ; //#todo-vulkanddpi
 bSupportsArrayTextureCompression = false
 
-bSupportsVolumeTextureCompression = true
 bSupportsDistanceFields = true
 bSupportsDiaphragmDOF = true
 bSupportsRGBColorBuffer = true
diff --git a/Engine/Source/Developer/Android/AndroidTargetPlatform/Private/AndroidTargetPlatform.cpp b/Engine/Source/Developer/Android/AndroidTargetPlatform/Private/AndroidTargetPlatform.cpp
--- a/Engine/Source/Developer/Android/AndroidTargetPlatform/Private/AndroidTargetPlatform.cpp
+++ b/Engine/Source/Developer/Android/AndroidTargetPlatform/Private/AndroidTargetPlatform.cpp
@@ -414,6 +414,11 @@ const FStaticMeshLODSettings& FAndroidTargetPlatform::GetStaticMeshLODSettings(
 	return StaticMeshLODSettings;
 }
 
+static bool FormatSupportsCompressedVolumeTexture(EAndroidTextureFormatCategory FormatCategory)
+{
+	// Supported in ES3.2 with ASTC
+	return FormatCategory == EAndroidTextureFormatCategory::ASTC;
+}
 
 void FAndroidTargetPlatform::GetTextureFormats( const UTexture* InTexture, TArray< TArray<FName> >& OutFormats) const
 {
@@ -432,6 +437,8 @@ void FAndroidTargetPlatform::GetTextureFormats( const UTexture* InTexture, TArra
 			continue;
 		}
 
+		const bool bSupportCompressedVolumeTexture = FormatSupportsCompressedVolumeTexture(FormatCategory);
+
 		TArray<FName> FormatPerLayer;
 		FormatPerLayer.SetNum(NumLayers);
 
@@ -448,7 +455,8 @@ void FAndroidTargetPlatform::GetTextureFormats( const UTexture* InTexture, TArra
 				|| (InTexture->Source.GetSizeX() < 4)						// Don't compress textures smaller than the DXT block size.
 				|| (InTexture->Source.GetSizeY() < 4)
 				|| (InTexture->Source.GetSizeX() % 4 != 0)
-				|| (InTexture->Source.GetSizeY() % 4 != 0);
+				|| (InTexture->Source.GetSizeY() % 4 != 0)
+				|| (InTexture->GetMaterialType() == MCT_VolumeTexture && !bSupportCompressedVolumeTexture);
 
 			// Determine the pixel format of the compressed texture.
 			if (InTexture->LODGroup == TEXTUREGROUP_Shadowmap)
diff --git a/Engine/Source/Developer/Android/AndroidTargetPlatform/Private/AndroidTargetPlatform.h b/Engine/Source/Developer/Android/AndroidTargetPlatform/Private/AndroidTargetPlatform.h
--- a/Engine/Source/Developer/Android/AndroidTargetPlatform/Private/AndroidTargetPlatform.h
+++ b/Engine/Source/Developer/Android/AndroidTargetPlatform/Private/AndroidTargetPlatform.h
@@ -351,7 +351,9 @@ public:
 			{ { FName(TEXT("AutoDXT")) },	{ FName(TEXT("ASTC_RGBAuto")) } },
 		};
 
-		GetDefaultTextureFormatNamePerLayer(OutFormats.AddDefaulted_GetRef(), this, Texture, true, false, 1);
+		// Supported in ES3.2 with ASTC
+		bool bSupportCompressedVolumeTexture = true;
+		GetDefaultTextureFormatNamePerLayer(OutFormats.AddDefaulted_GetRef(), this, Texture, true, bSupportCompressedVolumeTexture, 1);
 
 		for (FName& TextureFormatName : OutFormats.Last())
 		{
diff --git a/Engine/Source/Developer/IOS/IOSTargetPlatform/Private/IOSTargetPlatform.cpp b/Engine/Source/Developer/IOS/IOSTargetPlatform/Private/IOSTargetPlatform.cpp
--- a/Engine/Source/Developer/IOS/IOSTargetPlatform/Private/IOSTargetPlatform.cpp
+++ b/Engine/Source/Developer/IOS/IOSTargetPlatform/Private/IOSTargetPlatform.cpp
@@ -586,7 +586,8 @@ void FIOSTargetPlatform::GetTextureFormats( const UTexture* Texture, TArray< TAr
 	if (TextureFormatNames.Num() == 0)
 	{
         int32 BlockSize = 1;
-		GetDefaultTextureFormatNamePerLayer(TextureFormatNames, this, Texture, true, false, BlockSize);
+		bool bSupportCompressedVolumeTexture = true;
+		GetDefaultTextureFormatNamePerLayer(TextureFormatNames, this, Texture, true, bSupportCompressedVolumeTexture, BlockSize);
 	}
 
 	// include the formats we want
diff --git a/Engine/Source/Developer/Windows/WindowsTargetPlatform/Private/GenericWindowsTargetPlatform.h b/Engine/Source/Developer/Windows/WindowsTargetPlatform/Private/GenericWindowsTargetPlatform.h
--- a/Engine/Source/Developer/Windows/WindowsTargetPlatform/Private/GenericWindowsTargetPlatform.h
+++ b/Engine/Source/Developer/Windows/WindowsTargetPlatform/Private/GenericWindowsTargetPlatform.h
@@ -81,10 +81,6 @@ public:
 			{
 				bSupportDX11TextureFormats = false;
 			}
-			if (!UVolumeTexture::ShaderPlatformSupportsCompression(ShaderPlatform))
-			{
-				bSupportCompressedVolumeTexture = false;
-			}
 		}
 
 		// If we are targeting ES3.1, we also must cook encoded HDR reflection captures
diff --git a/Engine/Source/Runtime/Engine/Classes/Engine/VolumeTexture.h b/Engine/Source/Runtime/Engine/Classes/Engine/VolumeTexture.h
--- a/Engine/Source/Runtime/Engine/Classes/Engine/VolumeTexture.h
+++ b/Engine/Source/Runtime/Engine/Classes/Engine/VolumeTexture.h
@@ -121,9 +121,6 @@ public:
 
 #endif
 
-	ENGINE_API static bool ShaderPlatformSupportsCompression(FStaticShaderPlatform ShaderPlatform);
-
-
 	//~ Begin UStreamableRenderAsset Interface
 	virtual int32 CalcCumulativeLODSize(int32 NumLODs) const final override { return CalcTextureMemorySize(NumLODs); }
 	virtual bool StreamOut(int32 NewMipCount) final override;
diff --git a/Engine/Source/Runtime/Engine/Private/VolumeTexture.cpp b/Engine/Source/Runtime/Engine/Private/VolumeTexture.cpp
--- a/Engine/Source/Runtime/Engine/Private/VolumeTexture.cpp
+++ b/Engine/Source/Runtime/Engine/Private/VolumeTexture.cpp
@@ -282,8 +282,7 @@ uint32 UVolumeTexture::CalcTextureMemorySizeEnum( ETextureMipCount Enum ) const
 FTextureResource* UVolumeTexture::CreateResource()
 {
 	const FPixelFormatInfo& FormatInfo = GPixelFormats[GetPixelFormat()];
-	const bool bCompressedFormat = FormatInfo.BlockSizeX > 1; 
-	const bool bFormatIsSupported = FormatInfo.Supported && (!bCompressedFormat || ShaderPlatformSupportsCompression(GMaxRHIShaderPlatform));
+	const bool bFormatIsSupported = FormatInfo.Supported;
 
 	if (GetNumMips() > 0 && GSupportsTexture3D && bFormatIsSupported)
 	{
@@ -378,19 +377,6 @@ void UVolumeTexture::UpdateMipGenSettings()
 
 #endif // #if WITH_EDITOR
 
-bool UVolumeTexture::ShaderPlatformSupportsCompression(FStaticShaderPlatform ShaderPlatform)
-{
-	switch (ShaderPlatform)
-	{
-	case SP_PCD3D_SM5:
-	case SP_VULKAN_SM5_LUMIN:
-		return true;
-
-	default:
-		return FDataDrivenShaderPlatformInfo::GetSupportsVolumeTextureCompression(ShaderPlatform);
-	}
-}
-
 bool UVolumeTexture::StreamOut(int32 NewMipCount)
 {
 	FTexture3DResource* Texture3DResource = Resource ? Resource->GetTexture3DResource() : nullptr;
diff --git a/Engine/Source/Runtime/RHI/Private/RHI.cpp b/Engine/Source/Runtime/RHI/Private/RHI.cpp
--- a/Engine/Source/Runtime/RHI/Private/RHI.cpp
+++ b/Engine/Source/Runtime/RHI/Private/RHI.cpp
@@ -1656,7 +1656,6 @@ void FGenericDataDrivenShaderPlatformInfo::ParseDataDrivenShaderInfo(const FConf
 	Info.bIsAndroidOpenGLES = GetSectionBool(Section, "bIsAndroidOpenGLES");
 	Info.bSupportsMobileMultiView = GetSectionBool(Section, "bSupportsMobileMultiView");
 	Info.bSupportsArrayTextureCompression = GetSectionBool(Section, "bSupportsArrayTextureCompression");
-	Info.bSupportsVolumeTextureCompression = GetSectionBool(Section, "bSupportsVolumeTextureCompression");
 	Info.bSupportsDistanceFields = GetSectionBool(Section, "bSupportsDistanceFields");
 	Info.bSupportsDiaphragmDOF = GetSectionBool(Section, "bSupportsDiaphragmDOF");
 	Info.bSupportsRGBColorBuffer = GetSectionBool(Section, "bSupportsRGBColorBuffer");
diff --git a/Engine/Source/Runtime/RHI/Public/RHIDefinitions.h b/Engine/Source/Runtime/RHI/Public/RHIDefinitions.h
--- a/Engine/Source/Runtime/RHI/Public/RHIDefinitions.h
+++ b/Engine/Source/Runtime/RHI/Public/RHIDefinitions.h
@@ -325,7 +325,6 @@ class RHI_API FGenericDataDrivenShaderPlatformInfo
 
 	uint32 bSupportsMobileMultiView: 1;
 	uint32 bSupportsArrayTextureCompression : 1;
-	uint32 bSupportsVolumeTextureCompression: 1;
 	uint32 bSupportsDistanceFields: 1; // used for DFShadows and DFAO - since they had the same checks
 	uint32 bSupportsDiaphragmDOF: 1;
 	uint32 bSupportsRGBColorBuffer: 1;
@@ -455,11 +454,6 @@ public:
 		return Infos[Platform].bSupportsArrayTextureCompression;
 	}
 
-	static FORCEINLINE_DEBUGGABLE const bool GetSupportsVolumeTextureCompression(const FStaticShaderPlatform Platform)
-	{
-		return Infos[Platform].bSupportsVolumeTextureCompression;
-	}
-
 	static FORCEINLINE_DEBUGGABLE const bool GetSupportsDistanceFields(const FStaticShaderPlatform Platform)
 	{
 		return Infos[Platform].bSupportsDistanceFields;
