No unity fixes
#rnx

diff --git a/Engine/Source/Developer/Android/AndroidDeviceDetection/Private/AndroidDeviceDetectionModule.cpp b/Engine/Source/Developer/Android/AndroidDeviceDetection/Private/AndroidDeviceDetectionModule.cpp
--- a/Engine/Source/Developer/Android/AndroidDeviceDetection/Private/AndroidDeviceDetectionModule.cpp
+++ b/Engine/Source/Developer/Android/AndroidDeviceDetection/Private/AndroidDeviceDetectionModule.cpp
@@ -14,6 +14,7 @@
 #include "HAL/PlatformProcess.h"
 #include "Logging/LogMacros.h"
 #include "HAL/FileManager.h"
+#include "HAL/IConsoleManager.h"
 #include "Misc/Parse.h"
 #include "Misc/Paths.h"
 #include "HAL/Runnable.h"
diff --git a/Engine/Source/Runtime/Core/Private/Tests/Async/TaskGraphTest.cpp b/Engine/Source/Runtime/Core/Private/Tests/Async/TaskGraphTest.cpp
--- a/Engine/Source/Runtime/Core/Private/Tests/Async/TaskGraphTest.cpp
+++ b/Engine/Source/Runtime/Core/Private/Tests/Async/TaskGraphTest.cpp
@@ -15,6 +15,7 @@
 #include "Containers/Queue.h"
 #include "Tests/Benchmark.h"
 #include "HAL/Thread.h"
+#include "Async/Fundamental/Scheduler.h"
 
 #include <atomic>
 
