Fixed Android port for Unreal Insights

diff --git a/Engine/Source/Developer/Android/AndroidDeviceDetection/Private/AndroidDeviceDetectionModule.cpp b/Engine/Source/Developer/Android/AndroidDeviceDetection/Private/AndroidDeviceDetectionModule.cpp
--- a/Engine/Source/Developer/Android/AndroidDeviceDetection/Private/AndroidDeviceDetectionModule.cpp
+++ b/Engine/Source/Developer/Android/AndroidDeviceDetection/Private/AndroidDeviceDetectionModule.cpp
@@ -610,7 +610,7 @@ private:
 				// Add reverse port forwarding
 				uint16 ReversePortMappings[] = {
 					41899,	// Network file server, DEFAULT_TCP_FILE_SERVING_PORT in NetworkMessage.h
-					1980,	// Unreal Insights data collection, TraceInsightsModule.cpp
+					1981,	// Unreal Insights data collection, TraceInsightsModule.cpp
 					0		// end of list
 					};
 
