Fixed: BC4 texture format fallback for Android

diff --git a/Engine/Source/Developer/Android/AndroidTargetPlatform/Private/AndroidTargetPlatform.cpp b/Engine/Source/Developer/Android/AndroidTargetPlatform/Private/AndroidTargetPlatform.cpp
--- a/Engine/Source/Developer/Android/AndroidTargetPlatform/Private/AndroidTargetPlatform.cpp
+++ b/Engine/Source/Developer/Android/AndroidTargetPlatform/Private/AndroidTargetPlatform.cpp
@@ -50,7 +50,67 @@ struct FAndroidDeviceInfo;
 enum class ETargetPlatformFeatures;
 template<typename TPlatformProperties> class TTargetPlatformBase;
 
+namespace AndroidTexFormat
+{
+	// Compressed Texture Formats
+	const static FName NameDXT1(TEXT("DXT1"));
+	const static FName NameDXT3(TEXT("DXT3"));
+	const static FName NameDXT5(TEXT("DXT5"));
+	const static FName NameDXT5n(TEXT("DXT5n"));
+	const static FName NameAutoDXT(TEXT("AutoDXT"));
+	const static FName NameBC4(TEXT("BC4"));
+	const static FName NameBC5(TEXT("BC5"));
+	const static FName NameBC6H(TEXT("BC6H"));
+	const static FName NameBC7(TEXT("BC7"));
+	const static FName NameETC2_RGB(TEXT("ETC2_RGB"));
+	const static FName NameETC2_RGBA(TEXT("ETC2_RGBA"));
+	const static FName NameAutoETC2(TEXT("AutoETC2"));
+	const static FName NameASTC_4x4(TEXT("ASTC_4x4"));
+	const static FName NameASTC_6x6(TEXT("ASTC_6x6"));
+	const static FName NameASTC_8x8(TEXT("ASTC_8x8"));
+	const static FName NameASTC_10x10(TEXT("ASTC_10x10"));
+	const static FName NameASTC_12x12(TEXT("ASTC_12x12"));
+	const static FName NameAutoASTC(TEXT("ASTC_RGBAuto"));
+	
+	// Uncompressed Texture Formats
+	const static FName NameBGRA8(TEXT("BGRA8"));
+	const static FName NameG8(TEXT("G8"));
+	const static FName NameVU8(TEXT("VU8"));
+	const static FName NameRGBA16F(TEXT("RGBA16F"));
+	const static FName NameR16F(TEXT("R16F"));
+	const static FName NameR5G6B5(TEXT("R5G6B5"));
+	const static FName NameA1RGB555(TEXT("A1RGB555"));
+	//A1RGB555 is mapped to RGB555A1, because OpenGL GL_RGB5_A1 only supports alpha on the lowest bit
+	const static FName NameRGB555A1(TEXT("RGB555A1"));
+
+	const static FName ASTCRemap[][2] =
+	{
+		// Default format:		ASTC format:
+		{ NameDXT1,			FName(TEXT("ASTC_RGB"))		},
+		{ NameDXT5,			FName(TEXT("ASTC_RGBA"))	},
+		{ NameDXT5n,		FName(TEXT("ASTC_NormalAG"))},
+		{ NameBC5,			FName(TEXT("ASTC_NormalRG"))},
+		{ NameBC4,			NameG8						},
+		{ NameBC6H,			FName(TEXT("ASTC_RGB"))		},
+		{ NameBC7,			NameAutoASTC				},
+		{ NameAutoDXT,		NameAutoASTC				},
+		{ NameA1RGB555,		NameRGB555A1				}
+	};
 
+	const static FName ETCRemap[][2] =
+	{
+		// Default format:	ETC2 format:
+		{ NameDXT1,			NameETC2_RGB	},
+		{ NameDXT5,			NameETC2_RGBA	},
+		{ NameDXT5n,		NameETC2_RGB	},
+		{ NameBC5,			NameETC2_RGB	},
+		{ NameBC4,			NameG8			},
+		{ NameBC6H,			NameETC2_RGB	},
+		{ NameBC7,			NameAutoETC2	},
+		{ NameAutoDXT,		NameAutoETC2	},
+		{ NameA1RGB555,		NameRGB555A1	}
+	};
+}
 
 static FString GetLicensePath()
 {
@@ -402,150 +462,47 @@ const FStaticMeshLODSettings& FAndroidTargetPlatform::GetStaticMeshLODSettings(
 	return StaticMeshLODSettings;
 }
 
-static bool FormatSupportsCompressedVolumeTexture(EAndroidTextureFormatCategory FormatCategory)
+void FAndroidTargetPlatform::GetTextureFormats( const UTexture* Texture, TArray< TArray<FName> >& OutFormats) const
 {
+	check(Texture);
+	
 	// Supported in ES3.2 with ASTC
-	return FormatCategory == EAndroidTextureFormatCategory::ASTC;
-}
-
-void FAndroidTargetPlatform::GetTextureFormats( const UTexture* InTexture, TArray< TArray<FName> >& OutFormats) const
-{
-#if WITH_EDITOR
-	const int32 NumLayers = InTexture->Source.GetNumLayers();
+	const bool bSupportCompressedVolumeTexture = SupportsTextureFormatCategory(EAndroidTextureFormatCategory::ASTC);
+	// TODO: compressed HDR formats
+	const bool bSupportDX11TextureFormats = false;
 
-	// Can always compress power-of-two, sometimes support non-POT compression
-	const bool bIsCompressionValid = InTexture->Source.IsPowerOfTwo() || SupportsCompressedNonPOT();
+	TArray<FName>& LayerFormats = OutFormats.AddDefaulted_GetRef();
+	GetDefaultTextureFormatNamePerLayer(LayerFormats, this, Texture, bSupportDX11TextureFormats, bSupportCompressedVolumeTexture, 1);
 
-	OutFormats.Reserve((int32)EAndroidTextureFormatCategory::Count);
-	for (int32 FormatIndex = 0; FormatIndex < (int32)EAndroidTextureFormatCategory::Count; ++FormatIndex)
+	for (FName& TextureFormatName : LayerFormats)
 	{
-		const EAndroidTextureFormatCategory FormatCategory = (EAndroidTextureFormatCategory)FormatIndex;
-		if (!SupportsTextureFormatCategory(FormatCategory))
+		if (Texture->LODGroup == TEXTUREGROUP_Shadowmap)
 		{
-			continue;
+			// forward rendering only needs one channel for shadow maps
+			TextureFormatName = FName(TEXT("G8"));
 		}
 
-		const bool bSupportCompressedVolumeTexture = FormatSupportsCompressedVolumeTexture(FormatCategory);
-
-		TArray<FName> FormatPerLayer;
-		FormatPerLayer.SetNum(NumLayers);
-
-		bool bValidFormat = true;
-		for (int32 LayerIndex = 0; LayerIndex < NumLayers; ++LayerIndex)
+		if (Texture->IsA(UTextureCube::StaticClass()))
 		{
-			FTextureFormatSettings LayerFormatSettings;
-			InTexture->GetLayerFormatSettings(LayerIndex, LayerFormatSettings);
-
-			const bool bNoCompression = LayerFormatSettings.CompressionNone				// Code wants the texture uncompressed.
-				|| (InTexture->LODGroup == TEXTUREGROUP_ColorLookupTable)	// Textures in certain LOD groups should remain uncompressed.
-				|| (InTexture->LODGroup == TEXTUREGROUP_Bokeh)
-				|| (LayerFormatSettings.CompressionSettings == TC_EditorIcon)
-				|| (InTexture->Source.GetSizeX() < 4)						// Don't compress textures smaller than the DXT block size.
-				|| (InTexture->Source.GetSizeY() < 4)
-				|| (InTexture->Source.GetSizeX() % 4 != 0)
-				|| (InTexture->Source.GetSizeY() % 4 != 0)
-				|| (InTexture->GetMaterialType() == MCT_VolumeTexture && !bSupportCompressedVolumeTexture);
-
-			// Determine the pixel format of the compressed texture.
-			if (InTexture->LODGroup == TEXTUREGROUP_Shadowmap)
-			{
-				// forward rendering only needs one channel for shadow maps
-				FormatPerLayer[LayerIndex] = AndroidTexFormat::NameG8;
-			}
-			else if (bNoCompression && InTexture->HasHDRSource(LayerIndex))
-			{
-				FormatPerLayer[LayerIndex] = AndroidTexFormat::NameRGBA16F;
-			}
-			else if (bNoCompression)
-			{
-				FormatPerLayer[LayerIndex] = AndroidTexFormat::NameBGRA8;
-			}
-			else if (LayerFormatSettings.CompressionSettings == TC_LQ)
-			{
-				FormatPerLayer[LayerIndex] = LayerFormatSettings.CompressionNoAlpha ? AndroidTexFormat::NameR5G6B5 : AndroidTexFormat::NameRGB555A1;
-			}
-			else if (LayerFormatSettings.CompressionSettings == TC_EncodedReflectionCapture && !LayerFormatSettings.CompressionNone)
-			{
-				FormatPerLayer[LayerIndex] = AndroidTexFormat::NameETC2_RGBA;
-			}
-			else if (LayerFormatSettings.CompressionSettings == TC_HDR || LayerFormatSettings.CompressionSettings == TC_HDR_Compressed)
-			{
-				FormatPerLayer[LayerIndex] = AndroidTexFormat::NameRGBA16F;
-			}
-			else if (LayerFormatSettings.CompressionSettings == TC_Normalmap)
-			{
-				if(!bIsCompressionValid) FormatPerLayer[LayerIndex] = AndroidTexFormat::NamePOTERROR;
-				else if (FormatCategory == EAndroidTextureFormatCategory::DXT) FormatPerLayer[LayerIndex] = AndroidTexFormat::NameDXT5;
-				else if (FormatCategory == EAndroidTextureFormatCategory::ETC2) FormatPerLayer[LayerIndex] = AndroidTexFormat::NameETC2_RGB;
-				else bValidFormat = false;
-			}
-			else if (LayerFormatSettings.CompressionSettings == TC_Displacementmap)
-			{
-				FormatPerLayer[LayerIndex] = AndroidTexFormat::NameRGBA16F;
-			}
-			else if (LayerFormatSettings.CompressionSettings == TC_VectorDisplacementmap)
-			{
-				FormatPerLayer[LayerIndex] = AndroidTexFormat::NameBGRA8;
-			}
-			else if (LayerFormatSettings.CompressionSettings == TC_Grayscale)
-			{
-				FormatPerLayer[LayerIndex] = AndroidTexFormat::NameG8;
-			}
-			else if (LayerFormatSettings.CompressionSettings == TC_Alpha)
-			{
-				FormatPerLayer[LayerIndex] = AndroidTexFormat::NameG8;
-			}
-			else if (LayerFormatSettings.CompressionSettings == TC_DistanceFieldFont)
-			{
-				FormatPerLayer[LayerIndex] = AndroidTexFormat::NameG8;
-			}
-			else if (LayerFormatSettings.CompressionSettings == TC_HalfFloat)
+			const UTextureCube* Cube = CastChecked<UTextureCube>(Texture);
+			if (Cube != nullptr)
 			{
-				FormatPerLayer[LayerIndex] = AndroidTexFormat::NameR16F;
-			}
-			else if (LayerFormatSettings.CompressionSettings == TC_BC7)
-			{
-				if (!bIsCompressionValid) FormatPerLayer[LayerIndex] = AndroidTexFormat::NamePOTERROR;
-				else if (FormatCategory == EAndroidTextureFormatCategory::DXT) FormatPerLayer[LayerIndex] = AndroidTexFormat::NameDXT5;
-				else if (FormatCategory == EAndroidTextureFormatCategory::ETC2) FormatPerLayer[LayerIndex] = AndroidTexFormat::NameAutoETC2;
-				else bValidFormat = false;
-			}
-			else if (LayerFormatSettings.CompressionNoAlpha)
-			{
-				if (!bIsCompressionValid) FormatPerLayer[LayerIndex] = AndroidTexFormat::NamePOTERROR;
-				else if (FormatCategory == EAndroidTextureFormatCategory::DXT) FormatPerLayer[LayerIndex] = AndroidTexFormat::NameDXT1;
-				else if (FormatCategory == EAndroidTextureFormatCategory::ETC2) FormatPerLayer[LayerIndex] = AndroidTexFormat::NameETC2_RGB;
-				else bValidFormat = false;
-			}
-			else if (InTexture->bDitherMipMapAlpha)
-			{
-				if (!bIsCompressionValid) FormatPerLayer[LayerIndex] = AndroidTexFormat::NamePOTERROR;
-				else if (FormatCategory == EAndroidTextureFormatCategory::DXT) FormatPerLayer[LayerIndex] = AndroidTexFormat::NameDXT5;
-				else if (FormatCategory == EAndroidTextureFormatCategory::ETC2) FormatPerLayer[LayerIndex] = AndroidTexFormat::NameAutoETC2;
-				else bValidFormat = false;
-			}
-			else
-			{
-				if (!bIsCompressionValid) FormatPerLayer[LayerIndex] = AndroidTexFormat::NamePOTERROR;
-				else if (FormatCategory == EAndroidTextureFormatCategory::DXT) FormatPerLayer[LayerIndex] = AndroidTexFormat::NameAutoDXT;
-				else if (FormatCategory == EAndroidTextureFormatCategory::ETC2) FormatPerLayer[LayerIndex] = AndroidTexFormat::NameAutoETC2;
-				else bValidFormat = false;
+				FTextureFormatSettings FormatSettings;
+				Cube->GetDefaultFormatSettings(FormatSettings);
+				if (FormatSettings.CompressionSettings == TC_EncodedReflectionCapture && !FormatSettings.CompressionNone)
+				{
+					TextureFormatName = FName(TEXT("ETC2_RGBA"));
+				}
 			}
 		}
-
-		if (bValidFormat)
-		{
-			OutFormats.AddUnique(FormatPerLayer);
-		}
 	}
-#endif // WITH_EDITOR
 }
 
 FName FAndroidTargetPlatform::FinalizeVirtualTextureLayerFormat(FName Format) const
 {
 #if WITH_EDITOR
 	// Remap non-ETC variants to ETC
-	const static FName ETCRemap[][2] =
+	const static FName VTRemap[][2] =
 	{
 		{ { FName(TEXT("ASTC_RGB")) },			{ AndroidTexFormat::NameETC2_RGB } },
 		{ { FName(TEXT("ASTC_RGBA")) },			{ AndroidTexFormat::NameETC2_RGBA } },
@@ -557,11 +514,11 @@ FName FAndroidTargetPlatform::FinalizeVirtualTextureLayerFormat(FName Format) co
 		{ { AndroidTexFormat::NameAutoDXT },	{ AndroidTexFormat::NameAutoETC2 } }
 	};
 
-	for (int32 RemapIndex = 0; RemapIndex < UE_ARRAY_COUNT(ETCRemap); RemapIndex++)
+	for (int32 RemapIndex = 0; RemapIndex < UE_ARRAY_COUNT(VTRemap); RemapIndex++)
 	{
-		if (ETCRemap[RemapIndex][0] == Format)
+		if (VTRemap[RemapIndex][0] == Format)
 		{
-			return ETCRemap[RemapIndex][1];
+			return VTRemap[RemapIndex][1];
 		}
 	}
 #endif
@@ -570,29 +527,76 @@ FName FAndroidTargetPlatform::FinalizeVirtualTextureLayerFormat(FName Format) co
 
 void FAndroidTargetPlatform::GetAllTextureFormats(TArray<FName>& OutFormats) const
 {
-	OutFormats.Add(AndroidTexFormat::NameG8);
-	OutFormats.Add(AndroidTexFormat::NameRGBA16F);
-	OutFormats.Add(AndroidTexFormat::NameBGRA8);
-	OutFormats.Add(AndroidTexFormat::NameRGBA16F);
-	OutFormats.Add(AndroidTexFormat::NameRGBA16F);
-	OutFormats.Add(AndroidTexFormat::NameBGRA8);
-	OutFormats.Add(AndroidTexFormat::NameG8);
-	OutFormats.Add(AndroidTexFormat::NameG8);
-	OutFormats.Add(AndroidTexFormat::NameG8);
-	OutFormats.Add(AndroidTexFormat::NameR16F);
-
-	auto AddAllTextureFormatIfSupports = [=, &OutFormats](bool bIsNonPOT)
-	{
-		AddTextureFormatIfSupports(AndroidTexFormat::NameAutoDXT, OutFormats, bIsNonPOT);
-		AddTextureFormatIfSupports(AndroidTexFormat::NameDXT1, OutFormats, bIsNonPOT);
-		AddTextureFormatIfSupports(AndroidTexFormat::NameDXT5, OutFormats, bIsNonPOT);
-		AddTextureFormatIfSupports(AndroidTexFormat::NameAutoETC2, OutFormats, bIsNonPOT);
-	};
+	GetAllDefaultTextureFormats(this, OutFormats, false);
+}
 
-	AddAllTextureFormatIfSupports(true);
-	AddAllTextureFormatIfSupports(false);
+void FAndroid_ASTCTargetPlatform::GetAllTextureFormats(TArray<FName>& OutFormats) const
+{
+	GetAllDefaultTextureFormats(this, OutFormats, false);
+
+	// not supported
+	OutFormats.Remove(AndroidTexFormat::NameDXT3);
+	for (int32 RemapIndex = 0; RemapIndex < UE_ARRAY_COUNT(AndroidTexFormat::ASTCRemap); ++RemapIndex)
+	{
+		OutFormats.Remove(AndroidTexFormat::ASTCRemap[RemapIndex][0]);
+	}
+	
+	// support only ASTC for compressed textures
+	OutFormats.Add(AndroidTexFormat::NameAutoASTC);
 }
 
+void FAndroid_ASTCTargetPlatform::GetTextureFormats(const UTexture* Texture, TArray< TArray<FName> >& OutFormats) const
+{
+	FAndroidTargetPlatform::GetTextureFormats(Texture, OutFormats);
+
+	// perform any remapping away from defaults
+	TArray<FName>& LayerFormats = OutFormats.Last();
+	for (FName& TextureFormatName : LayerFormats)
+	{
+		for (int32 RemapIndex = 0; RemapIndex < UE_ARRAY_COUNT(AndroidTexFormat::ASTCRemap); ++RemapIndex)
+		{
+			if (TextureFormatName == AndroidTexFormat::ASTCRemap[RemapIndex][0])
+			{
+				TextureFormatName = AndroidTexFormat::ASTCRemap[RemapIndex][1];
+				break;
+			}
+		}
+	}
+}
+
+void FAndroid_ETC2TargetPlatform::GetAllTextureFormats(TArray<FName>& OutFormats) const
+{
+	GetAllDefaultTextureFormats(this, OutFormats, false);
+
+	// not supported
+	OutFormats.Remove(AndroidTexFormat::NameDXT3);
+	for (int32 RemapIndex = 0; RemapIndex < UE_ARRAY_COUNT(AndroidTexFormat::ETCRemap); ++RemapIndex)
+	{
+		OutFormats.Remove(AndroidTexFormat::ETCRemap[RemapIndex][0]);
+	}
+	
+	// support only ETC for compressed textures
+	OutFormats.Add(AndroidTexFormat::NameAutoETC2);
+}
+
+void FAndroid_ETC2TargetPlatform::GetTextureFormats(const UTexture* Texture, TArray< TArray<FName> >& OutFormats) const
+{
+	FAndroidTargetPlatform::GetTextureFormats(Texture, OutFormats);
+
+	// perform any remapping away from defaults
+	TArray<FName>& LayerFormats = OutFormats.Last();
+	for (FName& TextureFormatName : LayerFormats)
+	{
+		for (int32 RemapIndex = 0; RemapIndex < UE_ARRAY_COUNT(AndroidTexFormat::ETCRemap); ++RemapIndex)
+		{
+			if (TextureFormatName == AndroidTexFormat::ETCRemap[RemapIndex][0])
+			{
+				TextureFormatName = AndroidTexFormat::ETCRemap[RemapIndex][1];
+				break;
+			}
+		}
+	}
+}
 
 void FAndroidTargetPlatform::GetReflectionCaptureFormats( TArray<FName>& OutFormats ) const
 {
@@ -609,7 +613,6 @@ void FAndroidTargetPlatform::GetReflectionCaptureFormats( TArray<FName>& OutForm
 	OutFormats.Add(FName(TEXT("EncodedHDR")));
 }
 
-
 const UTextureLODSettings& FAndroidTargetPlatform::GetTextureLODSettings() const
 {
 	return *TextureLODSettings;
@@ -680,22 +683,6 @@ bool FAndroidTargetPlatform::SupportsVariants() const
 
 /* FAndroidTargetPlatform implementation
  *****************************************************************************/
-
-void FAndroidTargetPlatform::AddTextureFormatIfSupports( FName Format, TArray<FName>& OutFormats, bool bIsCompressedNonPOT ) const
-{
-	if (SupportsTextureFormat(Format))
-	{
-		if (bIsCompressedNonPOT && SupportsCompressedNonPOT() == false)
-		{
-			OutFormats.Add(AndroidTexFormat::NamePOTERROR);
-		}
-		else
-		{
-			OutFormats.Add(Format);
-		}
-	}
-}
-
 void FAndroidTargetPlatform::InitializeDeviceDetection()
 {
 	DeviceDetection = FModuleManager::LoadModuleChecked<IAndroidDeviceDetectionModule>("AndroidDeviceDetection").GetAndroidDeviceDetection();
diff --git a/Engine/Source/Developer/Android/AndroidTargetPlatform/Private/AndroidTargetPlatform.h b/Engine/Source/Developer/Android/AndroidTargetPlatform/Private/AndroidTargetPlatform.h
--- a/Engine/Source/Developer/Android/AndroidTargetPlatform/Private/AndroidTargetPlatform.h
+++ b/Engine/Source/Developer/Android/AndroidTargetPlatform/Private/AndroidTargetPlatform.h
@@ -41,39 +41,6 @@ template<typename KeyType,typename ValueType,typename SetAllocator ,typename Key
 template<typename KeyType,typename ValueType,typename SetAllocator ,typename KeyFuncs > class TMultiMap;
 template<typename TPlatformProperties> class TTargetPlatformBase;
 
-/**
- * Defines supported texture format names.
- */
-namespace AndroidTexFormat
-{
-	// Compressed Texture Formats
-	static FName NameDXT1(TEXT("DXT1"));
-	static FName NameDXT5(TEXT("DXT5"));
-	static FName NameAutoDXT(TEXT("AutoDXT"));
-	static FName NameETC2_RGB(TEXT("ETC2_RGB"));
-	static FName NameETC2_RGBA(TEXT("ETC2_RGBA"));
-	static FName NameAutoETC2(TEXT("AutoETC2"));
-	static FName NameASTC_4x4(TEXT("ASTC_4x4"));
-	static FName NameASTC_6x6(TEXT("ASTC_6x6"));
-	static FName NameASTC_8x8(TEXT("ASTC_8x8"));
-	static FName NameASTC_10x10(TEXT("ASTC_10x10"));
-	static FName NameASTC_12x12(TEXT("ASTC_12x12"));
-	static FName NameAutoASTC(TEXT("AutoASTC"));
-
-	// Uncompressed Texture Formats
-	static FName NameBGRA8(TEXT("BGRA8"));
-	static FName NameG8(TEXT("G8"));
-	static FName NameVU8(TEXT("VU8"));
-	static FName NameRGBA16F(TEXT("RGBA16F"));
-	static FName NameR16F(TEXT("R16F"));
-	static FName NameR5G6B5(TEXT("R5G6B5"));
-	static FName NameA1RGB555(TEXT("A1RGB555"));
-	//A1RGB555 is mapped to RGB555A1, because OpenGL GL_RGB5_A1 only supports alpha on the lowest bit
-	static FName NameRGB555A1(TEXT("RGB555A1"));
-	// Error "formats" (uncompressed)
-	static FName NamePOTERROR(TEXT("POTERROR"));
-}
-
 /** Listed in order of priority...if device supports multiple formats, first format in list will be chosen */
 enum class EAndroidTextureFormatCategory
 {
@@ -141,12 +108,6 @@ public:
 
 	virtual bool SupportsFeature( ETargetPlatformFeatures Feature ) const override;
 
-	virtual bool SupportsTextureFormat( FName Format ) const 
-	{
-		// By default we support all texture formats.
-		return true;
-	}
-
 	virtual bool SupportsTextureFormatCategory(EAndroidTextureFormatCategory Category) const
 	{
 		return true;
@@ -206,16 +167,6 @@ public:
 	}
 
 protected:
-
-	/**
-	 * Adds the specified texture format to the OutFormats if this android target platforms supports it.
-	 *
-	 * @param Format - The format to add.
-	 * @param OutFormats - The collection of formats to add to.
-	 * @param bIsCompressedNonPOT - If this is true, the texture wants to be compressed but is not a power of 2
-	 */
-	void AddTextureFormatIfSupports( FName Format, TArray<FName>& OutFormats, bool bIsCompressedNonPOT=false ) const;
-
 	/**
 	 * Return true if this device has a supported set of extensions for this platform.
 	 *
@@ -279,17 +230,6 @@ public:
 	{
 	}
 
-	virtual bool SupportsTextureFormat(FName Format) const override
-	{
-		if (Format == AndroidTexFormat::NameDXT1 ||
-			Format == AndroidTexFormat::NameDXT5 ||
-			Format == AndroidTexFormat::NameAutoDXT)
-		{
-			return true;
-		}
-		return false;
-	}
-
 	virtual bool SupportsTextureFormatCategory(EAndroidTextureFormatCategory Category) const override
 	{
 		return Category == EAndroidTextureFormatCategory::DXT;
@@ -308,7 +248,6 @@ public:
 	}
 };
 
-
 class FAndroid_ASTCTargetPlatform : public FAndroidTargetPlatform
 {
 public:
@@ -316,111 +255,14 @@ public:
 	{
 	}
 
-	virtual bool SupportsTextureFormat(FName Format) const override
-	{
-		if (Format == AndroidTexFormat::NameASTC_4x4 ||
-			Format == AndroidTexFormat::NameASTC_6x6 ||
-			Format == AndroidTexFormat::NameASTC_8x8 ||
-			Format == AndroidTexFormat::NameASTC_10x10 ||
-			Format == AndroidTexFormat::NameASTC_12x12 ||
-			Format == AndroidTexFormat::NameAutoASTC)
-		{
-			return true;
-		}
-		return false;
-	}
-
 	virtual bool SupportsTextureFormatCategory(EAndroidTextureFormatCategory Category) const override
 	{
 		return Category == EAndroidTextureFormatCategory::ASTC;
 	}
 
 #if WITH_ENGINE
-	virtual void GetTextureFormats(const UTexture* Texture, TArray< TArray<FName> >& OutFormats) const
-	{
-		check(Texture);
-
-		// we remap some of the defaults (with ASTC formats)
-		static FName FormatRemap[][2] =
-		{
-			// Default format:				ASTC format:
-			{ { FName(TEXT("DXT1")) },		{ FName(TEXT("ASTC_RGB")) } },
-			{ { FName(TEXT("DXT5")) },		{ FName(TEXT("ASTC_RGBA")) } },
-			{ { FName(TEXT("DXT5n")) },		{ FName(TEXT("ASTC_NormalAG")) } },
-			{ { FName(TEXT("BC5")) },		{ FName(TEXT("ASTC_NormalRG")) } },
-			{ { FName(TEXT("BC6H")) },		{ FName(TEXT("ASTC_RGB")) } },
-			{ { FName(TEXT("BC7")) },		{ FName(TEXT("ASTC_RGBAuto")) } },
-			{ { FName(TEXT("AutoDXT")) },	{ FName(TEXT("ASTC_RGBAuto")) } },
-			{ { AndroidTexFormat::NameA1RGB555 },	{ AndroidTexFormat::NameRGB555A1 } },
-		};
-
-		// Supported in ES3.2 with ASTC
-		bool bSupportCompressedVolumeTexture = true;
-		GetDefaultTextureFormatNamePerLayer(OutFormats.AddDefaulted_GetRef(), this, Texture, true, bSupportCompressedVolumeTexture, 1);
-
-		for (FName& TextureFormatName : OutFormats.Last())
-		{
-			if (Texture->LODGroup == TEXTUREGROUP_Shadowmap)
-			{
-				// forward rendering only needs one channel for shadow maps
-				TextureFormatName = FName(TEXT("G8"));
-			}
-			else
-			{
-				// perform any remapping away from defaults
-				for (int32 RemapIndex = 0; RemapIndex < UE_ARRAY_COUNT(FormatRemap); ++RemapIndex)
-				{
-					if (TextureFormatName == FormatRemap[RemapIndex][0])
-					{
-						// we found a remapping
-						TextureFormatName = FormatRemap[RemapIndex][1];
-						break;
-					}
-				}
-			}
-
-
-			if (Texture->IsA(UTextureCube::StaticClass()))
-			{
-				const UTextureCube* Cube = CastChecked<UTextureCube>(Texture);
-				if (Cube != nullptr)
-				{
-					FTextureFormatSettings FormatSettings;
-					Cube->GetDefaultFormatSettings(FormatSettings);
-					if (FormatSettings.CompressionSettings == TC_EncodedReflectionCapture && !FormatSettings.CompressionNone)
-					{
-						TextureFormatName = FName(TEXT("ETC2_RGBA"));
-					}
-				}
-			}
-		}
-	}
-
-
-	virtual void GetAllTextureFormats(TArray<FName>& OutFormats) const override
-	{
-		// we remap some of the defaults (with ASTC formats)
-		static FName FormatRemap[][2] =
-		{
-			// Default format:				ASTC format:
-			{ { FName(TEXT("DXT1")) },		{ FName(TEXT("ASTC_RGB")) } },
-			{ { FName(TEXT("DXT5")) },		{ FName(TEXT("ASTC_RGBA")) } },
-			{ { FName(TEXT("DXT5n")) },		{ FName(TEXT("ASTC_NormalAG")) } },
-			{ { FName(TEXT("BC5")) },		{ FName(TEXT("ASTC_NormalRG")) } },
-			{ { FName(TEXT("BC6H")) },		{ FName(TEXT("ASTC_RGB")) } },
-			{ { FName(TEXT("BC7")) },		{ FName(TEXT("ASTC_RGBAuto")) } },
-			{ { FName(TEXT("AutoDXT")) },	{ FName(TEXT("ASTC_RGBAuto")) } },
-		};
-
-		GetAllDefaultTextureFormats(this, OutFormats, true);
-
-		for (int32 RemapIndex = 0; RemapIndex < UE_ARRAY_COUNT(FormatRemap); ++RemapIndex)
-		{
-			OutFormats.Remove(FormatRemap[RemapIndex][0]);
-			OutFormats.AddUnique(FormatRemap[RemapIndex][1]);
-		}
-	}
-
+	virtual void GetTextureFormats(const UTexture* Texture, TArray< TArray<FName> >& OutFormats) const;
+	virtual void GetAllTextureFormats(TArray<FName>& OutFormats) const override;
 #endif
 
 	virtual bool SupportedByExtensionsString(const FString& ExtensionsString, const int GLESVersion) const override
@@ -445,23 +287,16 @@ public:
 	{
 	}
 
-	virtual bool SupportsTextureFormat(FName Format) const override
-	{
-		if (Format == AndroidTexFormat::NameETC2_RGB ||
-			Format == AndroidTexFormat::NameETC2_RGBA ||
-			Format == AndroidTexFormat::NameAutoETC2)
-		{
-			return true;
-		}
-
-		return false;
-	}
-
 	virtual bool SupportsTextureFormatCategory(EAndroidTextureFormatCategory Category) const override
 	{
 		return Category == EAndroidTextureFormatCategory::ETC2;
 	}
 
+#if WITH_ENGINE
+	virtual void GetTextureFormats(const UTexture* Texture, TArray< TArray<FName> >& OutFormats) const;
+	virtual void GetAllTextureFormats(TArray<FName>& OutFormats) const override;
+#endif
+
 	virtual bool SupportedByExtensionsString(const FString& ExtensionsString, const int GLESVersion) const override
 	{
 		return GLESVersion >= 0x30000;
