Fixing android platform specific compression settings

diff --git a/Engine/Source/Developer/Android/AndroidTargetPlatform/Private/AndroidTargetPlatform.cpp b/Engine/Source/Developer/Android/AndroidTargetPlatform/Private/AndroidTargetPlatform.cpp
--- a/Engine/Source/Developer/Android/AndroidTargetPlatform/Private/AndroidTargetPlatform.cpp
+++ b/Engine/Source/Developer/Android/AndroidTargetPlatform/Private/AndroidTargetPlatform.cpp
@@ -621,44 +621,10 @@ const UTextureLODSettings& FAndroidTargetPlatform::GetTextureLODSettings() const
 
 FName FAndroidTargetPlatform::GetWaveFormat(const class USoundWave* Wave) const
 {
-	// Platform specific codec name
-	static bool bFormatRead = false;
-	static FName NAME_FORMAT;
-	if (!bFormatRead)
-	{
-		bFormatRead = true;
-
-		FName AudioSetting;
-		{
-			FString AudioSettingStr;
-			if (!GConfig->GetString(TEXT("/Script/AndroidRuntimeSettings.AndroidRuntimeSettings"), TEXT("AndroidAudio"), AudioSettingStr, GEngineIni))
-			{
-				AudioSetting = *AudioSettingStr;
-			}
-		}
-
-#if WITH_OGGVORBIS
-		if (AudioSetting == Audio::NAME_OGG || AudioSetting == NAME_None)
-		{
-			NAME_FORMAT = Audio::NAME_OGG;
-		}
-#else
-		if (AudioSetting == Audio::NAME_OGG)
-		{
-			UE_LOG(LogAudio, Error, TEXT("Attemped to select Ogg Vorbis encoding when the cooker is built without Ogg Vorbis support."));
-		}
-#endif
-		else
-		{
-			// Otherwise return ADPCM as it'll either be option '2' or 'default' depending on WITH_OGGVORBIS config
-			NAME_FORMAT = Audio::NAME_ADPCM;
-		}
-	}
-
 	FName FormatName = Audio::ToName(Wave->GetSoundAssetCompressionType());
 	if (FormatName == Audio::NAME_PLATFORM_SPECIFIC)
 	{
-		FormatName = NAME_FORMAT;
+		FormatName = Audio::NAME_OGG;
 	}
 	return FormatName;
 }
diff --git a/Engine/Source/Runtime/Android/AudioMixerAndroid/Private/AudioMixerPlatformAndroid.cpp b/Engine/Source/Runtime/Android/AudioMixerAndroid/Private/AudioMixerPlatformAndroid.cpp
--- a/Engine/Source/Runtime/Android/AudioMixerAndroid/Private/AudioMixerPlatformAndroid.cpp
+++ b/Engine/Source/Runtime/Android/AudioMixerAndroid/Private/AudioMixerPlatformAndroid.cpp
@@ -425,16 +425,7 @@ namespace Audio
 
 		if (RuntimeFormat == Audio::NAME_PLATFORM_SPECIFIC)
 		{
-#if WITH_OGGVORBIS
-			if (InSoundWave->HasCompressedData(Audio::NAME_OGG))
-			{
-				RuntimeFormat = Audio::NAME_OGG;
-			}
-			else
-#endif // #if WITH_OGGVORBIS
-			{
-				RuntimeFormat = Audio::NAME_ADPCM;
-			}
+			RuntimeFormat = Audio::NAME_OGG;
 		}
 
 		return RuntimeFormat;
@@ -444,20 +435,17 @@ namespace Audio
 	{
 		ICompressedAudioInfo* Decoder = nullptr;
 
-#if WITH_OGGVORBIS
 		if (InRuntimeFormat == Audio::NAME_OGG)
 		{
 			Decoder = new FVorbisAudioInfo();
 		}
-		else
-#endif // #if WITH_OGGVORBIS
-#if WITH_BINK_AUDIO
-		if (InRuntimeFormat == Audio::NAME_BINKA)
+		else if (InRuntimeFormat == Audio::NAME_BINKA)
 		{
+#if WITH_BINK_AUDIO
 			Decoder = new FBinkAudioInfo();
-		} 
-		else
 #endif // WITH_BINK_AUDIO			
+		}
+		else
 		{
 			Decoder = Audio::CreateSoundAssetDecoder(InRuntimeFormat);
 		}
