Support BC4 textures using ETC2_R11 (instead of uncompressed G8) on iOS, Android, QNX

diff --git a/Engine/Source/Developer/Android/AndroidTargetPlatform/Private/AndroidTargetPlatform.cpp b/Engine/Source/Developer/Android/AndroidTargetPlatform/Private/AndroidTargetPlatform.cpp
--- a/Engine/Source/Developer/Android/AndroidTargetPlatform/Private/AndroidTargetPlatform.cpp
+++ b/Engine/Source/Developer/Android/AndroidTargetPlatform/Private/AndroidTargetPlatform.cpp
@@ -64,6 +64,7 @@ namespace AndroidTexFormat
 	const static FName NameBC7(TEXT("BC7"));
 	const static FName NameETC2_RGB(TEXT("ETC2_RGB"));
 	const static FName NameETC2_RGBA(TEXT("ETC2_RGBA"));
+	const static FName NameETC2_R11(TEXT("ETC2_R11"));
 	const static FName NameAutoETC2(TEXT("AutoETC2"));
 	const static FName NameASTC_4x4(TEXT("ASTC_4x4"));
 	const static FName NameASTC_6x6(TEXT("ASTC_6x6"));
@@ -90,7 +91,7 @@ namespace AndroidTexFormat
 		{ NameDXT5,			FName(TEXT("ASTC_RGBA"))	},
 		{ NameDXT5n,		FName(TEXT("ASTC_NormalAG"))},
 		{ NameBC5,			FName(TEXT("ASTC_NormalRG"))},
-		{ NameBC4,			NameG8						},
+		{ NameBC4,			NameETC2_R11				},
 		{ NameBC6H,			FName(TEXT("ASTC_RGB"))		},
 		{ NameBC7,			NameAutoASTC				},
 		{ NameAutoDXT,		NameAutoASTC				},
@@ -104,7 +105,7 @@ namespace AndroidTexFormat
 		{ NameDXT5,			NameETC2_RGBA	},
 		{ NameDXT5n,		NameETC2_RGB	},
 		{ NameBC5,			NameETC2_RGB	},
-		{ NameBC4,			NameG8			},
+		{ NameBC4,			NameETC2_R11	},
 		{ NameBC6H,			NameETC2_RGB	},
 		{ NameBC7,			NameAutoETC2	},
 		{ NameAutoDXT,		NameAutoETC2	},
@@ -541,8 +542,10 @@ void FAndroid_ASTCTargetPlatform::GetAllTextureFormats(TArray<FName>& OutFormats
 		OutFormats.Remove(AndroidTexFormat::ASTCRemap[RemapIndex][0]);
 	}
 	
-	// support only ASTC for compressed textures
+	// ASTC for compressed textures
 	OutFormats.Add(AndroidTexFormat::NameAutoASTC);
+	// ETC for ETC2_R11
+	OutFormats.Add(AndroidTexFormat::NameAutoETC2);
 }
 
 void FAndroid_ASTCTargetPlatform::GetTextureFormats(const UTexture* Texture, TArray< TArray<FName> >& OutFormats) const
diff --git a/Engine/Source/Developer/IOS/IOSTargetPlatform/Private/IOSTargetPlatform.cpp b/Engine/Source/Developer/IOS/IOSTargetPlatform/Private/IOSTargetPlatform.cpp
--- a/Engine/Source/Developer/IOS/IOSTargetPlatform/Private/IOSTargetPlatform.cpp
+++ b/Engine/Source/Developer/IOS/IOSTargetPlatform/Private/IOSTargetPlatform.cpp
@@ -563,7 +563,7 @@ static FName FormatRemap[] =
 	FName(TEXT("DXT5")),	FName(TEXT("ASTC_RGBA")),
 	FName(TEXT("DXT5n")),	FName(TEXT("ASTC_NormalAG")),
 	FName(TEXT("BC5")),		FName(TEXT("ASTC_NormalRG")),
-	FName(TEXT("BC4")),		FName(TEXT("G8")),
+	FName(TEXT("BC4")),		FName(TEXT("ETC2_R11")),
 	FName(TEXT("BC6H")),	FName(TEXT("ASTC_RGB")), 
 	FName(TEXT("BC7")),		FName(TEXT("ASTC_RGBAuto"))
 };
diff --git a/Engine/Source/Developer/TextureFormatETC2/Private/TextureFormatETC2.cpp b/Engine/Source/Developer/TextureFormatETC2/Private/TextureFormatETC2.cpp
--- a/Engine/Source/Developer/TextureFormatETC2/Private/TextureFormatETC2.cpp
+++ b/Engine/Source/Developer/TextureFormatETC2/Private/TextureFormatETC2.cpp
@@ -39,6 +39,7 @@ class FETC2TextureBuildFunction final : public FTextureBuildFunction
 #define ENUM_SUPPORTED_FORMATS(op) \
 	op(ETC2_RGB) \
 	op(ETC2_RGBA) \
+	op(ETC2_R11) \
 	op(AutoETC2)
 
 #define DECL_FORMAT_NAME(FormatName) static FName GTextureFormatName##FormatName = FName(TEXT(#FormatName));
@@ -116,6 +117,9 @@ static bool CompressImageUsingQonvert(
 	case PF_ETC2_RGBA:
 		DstImg.nFormat = Q_FORMAT_ETC2_RGBA8;
 		break;
+	case PF_ETC2_R11_EAC:
+		DstImg.nFormat = Q_FORMAT_EAC_R_UNSIGNED;
+		break;
 	default:
 		UE_LOG(LogTextureFormatETC2, Fatal, TEXT("Unsupported EPixelFormat for compression: %u"), (uint32)PixelFormat);
 		return false;
@@ -146,7 +150,7 @@ class FTextureFormatETC2 : public ITextureFormat
 		const struct FTextureBuildSettings* BuildSettings = nullptr
 	) const override
 	{
-		return 0;
+		return 1;
 	}
 
 	virtual FName GetEncoderName(FName Format) const override
@@ -182,6 +186,11 @@ class FTextureFormatETC2 : public ITextureFormat
 			return PF_ETC2_RGBA;
 		}
 
+		if (BuildSettings.TextureFormatName == GTextureFormatNameETC2_R11)
+		{
+			return PF_ETC2_R11_EAC;
+		}
+
 		UE_LOG(LogTextureFormatETC2, Fatal, TEXT("Unhandled texture format '%s' given to FTextureFormatAndroid::GetPixelFormatForImage()"), *BuildSettings.TextureFormatName.ToString());
 		return PF_Unknown;
 	}
