NOP comment about FinalizeVirtualTextureLayerFormat

diff --git a/Engine/Source/Developer/Android/AndroidTargetPlatform/Private/AndroidTargetPlatform.cpp b/Engine/Source/Developer/Android/AndroidTargetPlatform/Private/AndroidTargetPlatform.cpp
--- a/Engine/Source/Developer/Android/AndroidTargetPlatform/Private/AndroidTargetPlatform.cpp
+++ b/Engine/Source/Developer/Android/AndroidTargetPlatform/Private/AndroidTargetPlatform.cpp
@@ -106,7 +106,7 @@ namespace AndroidTexFormat
 		{ NameDXT5n,		NameETC2_RGB	},
 		{ NameBC5,			NameETC2_RGB	},
 		{ NameBC4,			NameETC2_R11	},
-		{ NameBC6H,			NameETC2_RGB	},
+		{ NameBC6H,			NameETC2_RGB	}, // @todo Oodle : uncompressed float?
 		{ NameBC7,			NameAutoETC2	},
 		{ NameAutoDXT,		NameAutoETC2	},
 		{ NameA1RGB555,		NameRGB555A1	}
@@ -503,6 +503,12 @@ FName FAndroidTargetPlatform::FinalizeVirtualTextureLayerFormat(FName Format) co
 {
 #if WITH_EDITOR
 	// Remap non-ETC variants to ETC
+
+	// VirtualTexture Format was already run through the ordinary texture remaps to change AutoDXT to ASTC or ETC
+	// this then runs again
+	// currently it forces all ASTC to ETC
+	// this is needed because the runtime virtual texture encoder only supports ETC
+
 	const static FName VTRemap[][2] =
 	{
 		{ { FName(TEXT("ASTC_RGB")) },			{ AndroidTexFormat::NameETC2_RGB } },
diff --git a/Engine/Source/Developer/IOS/IOSTargetPlatform/Private/IOSTargetPlatform.cpp b/Engine/Source/Developer/IOS/IOSTargetPlatform/Private/IOSTargetPlatform.cpp
--- a/Engine/Source/Developer/IOS/IOSTargetPlatform/Private/IOSTargetPlatform.cpp
+++ b/Engine/Source/Developer/IOS/IOSTargetPlatform/Private/IOSTargetPlatform.cpp
@@ -653,6 +653,12 @@ void FIOSTargetPlatform::GetAllTextureFormats(TArray<FName>& OutFormats) const
 FName FIOSTargetPlatform::FinalizeVirtualTextureLayerFormat(FName Format) const
 {
 #if WITH_EDITOR
+
+	// VirtualTexture Format was already run through the ordinary texture remaps to change AutoDXT to ASTC or ETC
+	// this then runs again
+	// currently it forces all ASTC to ETC
+	// this is needed because the runtime virtual texture encoder only supports ETC
+
 	const static FName NameETC2_RGB(TEXT("ETC2_RGB"));
 	const static FName NameETC2_RGBA(TEXT("ETC2_RGBA"));
 	const static FName NameAutoETC2(TEXT("AutoETC2"));
