Removed bSupportDX11TextureFormats, functionality is no longer necessary. Pushed functionality of bSupportDX11TextureFormats down to target platform code.

diff --git a/Engine/Source/Developer/Android/AndroidTargetPlatform/Private/AndroidTargetPlatform.cpp b/Engine/Source/Developer/Android/AndroidTargetPlatform/Private/AndroidTargetPlatform.cpp
--- a/Engine/Source/Developer/Android/AndroidTargetPlatform/Private/AndroidTargetPlatform.cpp
+++ b/Engine/Source/Developer/Android/AndroidTargetPlatform/Private/AndroidTargetPlatform.cpp
@@ -74,17 +74,23 @@ namespace AndroidTexFormat
 	const static FName NameAutoASTC(TEXT("ASTC_RGBAuto"));
 	
 	// Uncompressed Texture Formats
-	//const static FName NameBGRA8(TEXT("BGRA8"));
-	//const static FName NameG8(TEXT("G8"));
-	//const static FName NameVU8(TEXT("VU8"));
-	//const static FName NameRGBA16F(TEXT("RGBA16F"));
-	//const static FName NameR16F(TEXT("R16F"));
-	//const static FName NameR5G6B5(TEXT("R5G6B5"));
+	const static FName NameBGRA8(TEXT("BGRA8"));
+	const static FName NameG8(TEXT("G8"));
+	const static FName NameVU8(TEXT("VU8"));
+	const static FName NameRGBA16F(TEXT("RGBA16F"));
+	const static FName NameR16F(TEXT("R16F"));
+	const static FName NameR5G6B5(TEXT("R5G6B5"));
 
 	//A1RGB555 is mapped to RGB555A1, because OpenGL GL_RGB5_A1 only supports alpha on the lowest bit
 	const static FName NameA1RGB555(TEXT("A1RGB555"));
 	const static FName NameRGB555A1(TEXT("RGB555A1"));
 
+	const static FName GenericRemap[][2] =
+	{
+		{ NameBC6H,			NameRGBA16F					},
+		{ NameBC7,			NameBGRA8					},
+	};
+
 	const static FName ASTCRemap[][2] =
 	{
 		// Default format:		ASTC format:
@@ -93,8 +99,10 @@ namespace AndroidTexFormat
 		{ NameDXT5n,		FName(TEXT("ASTC_NormalAG"))},
 		{ NameBC5,			FName(TEXT("ASTC_NormalRG"))},
 		{ NameBC4,			NameETC2_R11				},
-		{ NameBC6H,			FName(TEXT("ASTC_RGB"))		},
-		{ NameBC7,			NameAutoASTC				},
+		{ NameBC6H,			NameRGBA16F					},
+		{ NameBC7,			NameBGRA8					},
+		//{ NameBC6H,			FName(TEXT("ASTC_RGB"))		},
+		//{ NameBC7,			NameAutoASTC				},
 		{ NameAutoDXT,		NameAutoASTC				},
 		{ NameA1RGB555,		NameRGB555A1				}
 	};
@@ -107,8 +115,10 @@ namespace AndroidTexFormat
 		{ NameDXT5n,		NameETC2_RGB	},
 		{ NameBC5,			NameETC2_RGB	},
 		{ NameBC4,			NameETC2_R11	},
-		{ NameBC6H,			NameETC2_RGB	}, // @todo Oodle : uncompressed float?
-		{ NameBC7,			NameAutoETC2	},
+		{ NameBC6H,			NameRGBA16F					},
+		{ NameBC7,			NameBGRA8					},
+		//{ NameBC6H,			NameETC2_RGB	},
+		//{ NameBC7,			NameAutoETC2	},
 		{ NameAutoDXT,		NameAutoETC2	},
 		{ NameA1RGB555,		NameRGB555A1	}
 	};
@@ -486,13 +496,11 @@ void FAndroidTargetPlatform::GetTextureFormats( const UTexture* Texture, TArray<
 	
 	// Supported in ES3.2 with ASTC
 	const bool bSupportCompressedVolumeTexture = SupportsTextureFormatCategory(EAndroidTextureFormatCategory::ASTC);
-	// TODO: compressed HDR formats
-	const bool bSupportDX11TextureFormats = false;
 	// OpenGL ES has F32 textures but doesn't allow linear filtering unless OES_texture_float_linear
 	const bool bSupportFilteredFloat32Textures = false;
 
 	TArray<FName>& LayerFormats = OutFormats.AddDefaulted_GetRef();
-	GetDefaultTextureFormatNamePerLayer(LayerFormats, this, Texture, bSupportDX11TextureFormats, bSupportCompressedVolumeTexture, 1, bSupportFilteredFloat32Textures);
+	GetDefaultTextureFormatNamePerLayer(LayerFormats, this, Texture, bSupportCompressedVolumeTexture, 1, bSupportFilteredFloat32Textures);
 
 	for (FName& TextureFormatName : LayerFormats)
 	{
@@ -515,6 +523,14 @@ void FAndroidTargetPlatform::GetTextureFormats( const UTexture* Texture, TArray<
 				}
 			}
 		}
+
+		for (int32 RemapIndex = 0; RemapIndex < UE_ARRAY_COUNT(AndroidTexFormat::GenericRemap); ++RemapIndex)
+		{
+			if(TextureFormatName == AndroidTexFormat::GenericRemap[RemapIndex][0]) 
+			{
+				TextureFormatName = AndroidTexFormat::GenericRemap[RemapIndex][1];
+			}
+		}
 	}
 }
 
@@ -553,12 +569,32 @@ FName FAndroidTargetPlatform::FinalizeVirtualTextureLayerFormat(FName Format) co
 
 void FAndroidTargetPlatform::GetAllTextureFormats(TArray<FName>& OutFormats) const
 {
-	GetAllDefaultTextureFormats(this, OutFormats, false);
+	GetAllDefaultTextureFormats(this, OutFormats);
+
+	for (int32 RemapIndex = 0; RemapIndex < UE_ARRAY_COUNT(AndroidTexFormat::GenericRemap); ++RemapIndex)
+	{
+		OutFormats.Remove(AndroidTexFormat::GenericRemap[RemapIndex][0]);
+	}
+
+	for (int32 RemapIndex = 0; RemapIndex < UE_ARRAY_COUNT(AndroidTexFormat::GenericRemap); ++RemapIndex)
+	{
+		OutFormats.AddUnique(AndroidTexFormat::GenericRemap[RemapIndex][1]);
+	}
 }
 
 void FAndroid_ASTCTargetPlatform::GetAllTextureFormats(TArray<FName>& OutFormats) const
 {
-	GetAllDefaultTextureFormats(this, OutFormats, false);
+	GetAllDefaultTextureFormats(this, OutFormats);
+
+	for (int32 RemapIndex = 0; RemapIndex < UE_ARRAY_COUNT(AndroidTexFormat::GenericRemap); ++RemapIndex)
+	{
+		OutFormats.Remove(AndroidTexFormat::GenericRemap[RemapIndex][0]);
+	}
+
+	for (int32 RemapIndex = 0; RemapIndex < UE_ARRAY_COUNT(AndroidTexFormat::GenericRemap); ++RemapIndex)
+	{
+		OutFormats.AddUnique(AndroidTexFormat::GenericRemap[RemapIndex][1]);
+	}
 
 	// not supported
 	OutFormats.Remove(AndroidTexFormat::NameDXT3);
@@ -594,7 +630,17 @@ void FAndroid_ASTCTargetPlatform::GetTextureFormats(const UTexture* Texture, TAr
 
 void FAndroid_ETC2TargetPlatform::GetAllTextureFormats(TArray<FName>& OutFormats) const
 {
-	GetAllDefaultTextureFormats(this, OutFormats, false);
+	GetAllDefaultTextureFormats(this, OutFormats);
+
+	for (int32 RemapIndex = 0; RemapIndex < UE_ARRAY_COUNT(AndroidTexFormat::GenericRemap); ++RemapIndex)
+	{
+		OutFormats.Remove(AndroidTexFormat::GenericRemap[RemapIndex][0]);
+	}
+
+	for (int32 RemapIndex = 0; RemapIndex < UE_ARRAY_COUNT(AndroidTexFormat::GenericRemap); ++RemapIndex)
+	{
+		OutFormats.AddUnique(AndroidTexFormat::GenericRemap[RemapIndex][1]);
+	}
 
 	// not supported
 	OutFormats.Remove(AndroidTexFormat::NameDXT3);
diff --git a/Engine/Source/Developer/HoloLens/HoloLensTargetPlatform/Private/HoloLensTargetPlatform.cpp b/Engine/Source/Developer/HoloLens/HoloLensTargetPlatform/Private/HoloLensTargetPlatform.cpp
--- a/Engine/Source/Developer/HoloLens/HoloLensTargetPlatform/Private/HoloLensTargetPlatform.cpp
+++ b/Engine/Source/Developer/HoloLens/HoloLensTargetPlatform/Private/HoloLensTargetPlatform.cpp
@@ -188,15 +188,45 @@ void FHoloLensTargetPlatform::GetReflectionCaptureFormats(TArray<FName>& OutForm
 	OutFormats.Add(FName(TEXT("EncodedHDR")));
 }
 
+const static FName NameBC6H(TEXT("BC6H"));
+const static FName NameBC7(TEXT("BC7"));
+const static FName NameRGBA16F(TEXT("RGBA16F"));
+const static FName NameBGRA8(TEXT("BGRA8"));
+const static FName HoloLensRemap[][2] =
+{
+	{ NameBC6H,			NameRGBA16F		},
+	{ NameBC7,			NameBGRA8		},
+};
+
 void FHoloLensTargetPlatform::GetTextureFormats(const UTexture* InTexture, TArray< TArray<FName> >& OutFormats) const
 {
 	bool bSupportFilteredFloat32Textures = false;
-	GetDefaultTextureFormatNamePerLayer(OutFormats.AddDefaulted_GetRef(), this, InTexture, false, false,4, bSupportFilteredFloat32Textures);
+	GetDefaultTextureFormatNamePerLayer(OutFormats.AddDefaulted_GetRef(), this, InTexture, false, 4, bSupportFilteredFloat32Textures);
+
+	// perform any remapping away from defaults
+	TArray<FName>& LayerFormats = OutFormats.Last();
+	for (FName& TextureFormatName : LayerFormats)
+	{
+		for (int32 RemapIndex = 0; RemapIndex < UE_ARRAY_COUNT(HoloLensRemap); ++RemapIndex)
+		{
+			if (TextureFormatName == HoloLensRemap[RemapIndex][0])
+			{
+				TextureFormatName = HoloLensRemap[RemapIndex][1];
+				break;
+			}
+		}
+	}
 }
 
 void FHoloLensTargetPlatform::GetAllTextureFormats(TArray<FName>& OutFormats) const
 {
-	GetAllDefaultTextureFormats(this, OutFormats, false);
+	GetAllDefaultTextureFormats(this, OutFormats);
+
+	// not supported
+	for (int32 RemapIndex = 0; RemapIndex < UE_ARRAY_COUNT(HoloLensRemap); ++RemapIndex)
+	{
+		OutFormats.Remove(HoloLensRemap[RemapIndex][0]);
+	}
 }
 
 #endif
diff --git a/Engine/Source/Developer/IOS/IOSTargetPlatform/Private/IOSTargetPlatform.cpp b/Engine/Source/Developer/IOS/IOSTargetPlatform/Private/IOSTargetPlatform.cpp
--- a/Engine/Source/Developer/IOS/IOSTargetPlatform/Private/IOSTargetPlatform.cpp
+++ b/Engine/Source/Developer/IOS/IOSTargetPlatform/Private/IOSTargetPlatform.cpp
@@ -609,7 +609,7 @@ void FIOSTargetPlatform::GetTextureFormats( const UTexture* Texture, TArray< TAr
 		// Compressed volume textures require MTLGPUFamilyApple3 or later
 		// min spec for TVOS is AppleTV HD which is MTLGPUFamilyApple2 (A8)
 		bool bSupportCompressedVolumeTexture = !bIsTVOS && !SupportsA8Devices();
-		GetDefaultTextureFormatNamePerLayer(TextureFormatNames, this, Texture, true, bSupportCompressedVolumeTexture, BlockSize);
+		GetDefaultTextureFormatNamePerLayer(TextureFormatNames, this, Texture, bSupportCompressedVolumeTexture, BlockSize);
 	}
 
 	// include the formats we want
@@ -653,7 +653,7 @@ void FIOSTargetPlatform::GetAllTextureFormats(TArray<FName>& OutFormats) const
 {
 	bool bFoundRemap = false;
 
-	GetAllDefaultTextureFormats(this, OutFormats, false);
+	GetAllDefaultTextureFormats(this, OutFormats);
 
 	for (int32 RemapIndex = 0; RemapIndex < UE_ARRAY_COUNT(FormatRemap); RemapIndex += 2)
 	{
diff --git a/Engine/Source/Developer/Linux/LinuxTargetPlatform/Private/LinuxTargetPlatform.h b/Engine/Source/Developer/Linux/LinuxTargetPlatform/Private/LinuxTargetPlatform.h
--- a/Engine/Source/Developer/Linux/LinuxTargetPlatform/Private/LinuxTargetPlatform.h
+++ b/Engine/Source/Developer/Linux/LinuxTargetPlatform/Private/LinuxTargetPlatform.h
@@ -320,7 +320,7 @@ public:
 		if (this->AllowAudioVisualData())
 		{
 			// just use the standard texture format name for this texture
-			GetAllDefaultTextureFormats(this, OutFormats, true);
+			GetAllDefaultTextureFormats(this, OutFormats);
 		}
 	}
 
diff --git a/Engine/Source/Developer/Windows/WindowsTargetPlatform/Private/GenericWindowsTargetPlatform.h b/Engine/Source/Developer/Windows/WindowsTargetPlatform/Private/GenericWindowsTargetPlatform.h
--- a/Engine/Source/Developer/Windows/WindowsTargetPlatform/Private/GenericWindowsTargetPlatform.h
+++ b/Engine/Source/Developer/Windows/WindowsTargetPlatform/Private/GenericWindowsTargetPlatform.h
@@ -71,7 +71,7 @@ public:
 		static FName NAME_SF_VULKAN_ES31(TEXT("SF_VULKAN_ES31"));
 		static FName NAME_OPENGL_150_ES3_1(TEXT("GLSL_150_ES31"));
 
-		bSupportDX11TextureFormats = true;
+		bool bSupportDX11TextureFormats = true;
 		bSupportCompressedVolumeTexture = true;
 
 		for (FName TargetedShaderFormat : TargetedShaderFormats)
@@ -94,6 +94,7 @@ public:
 				break;
 			}
 		}
+		check(bSupportDX11TextureFormats);
 
 		// If we are targeting ES3.1, we also must cook encoded HDR reflection captures
 		bRequiresEncodedHDRReflectionCaptures =	TargetedShaderFormats.Contains(NAME_SF_VULKAN_ES31)
@@ -306,7 +307,7 @@ public:
 	{
 		if (!TProperties::IsServerOnly())
 		{
-			GetDefaultTextureFormatNamePerLayer(OutFormats.AddDefaulted_GetRef(), this, InTexture, bSupportDX11TextureFormats, bSupportCompressedVolumeTexture);
+			GetDefaultTextureFormatNamePerLayer(OutFormats.AddDefaulted_GetRef(), this, InTexture, bSupportCompressedVolumeTexture);
 		}
 	}
 
@@ -314,7 +315,7 @@ public:
 	{
 		if (!TProperties::IsServerOnly())
 		{
-			GetAllDefaultTextureFormats(this, OutFormats, bSupportDX11TextureFormats);
+			GetAllDefaultTextureFormats(this, OutFormats);
 		}
 	}
 
@@ -409,9 +410,6 @@ private:
 	// Holds static mesh LOD settings.
 	FStaticMeshLODSettings StaticMeshLODSettings;
 
-	// True if the project supports non-DX11 texture formats.
-	bool bSupportDX11TextureFormats;
-
 	// True if the project requires encoded HDR reflection captures
 	bool bRequiresEncodedHDRReflectionCaptures;
 
diff --git a/Engine/Source/Runtime/Engine/Private/Texture.cpp b/Engine/Source/Runtime/Engine/Private/Texture.cpp
--- a/Engine/Source/Runtime/Engine/Private/Texture.cpp
+++ b/Engine/Source/Runtime/Engine/Private/Texture.cpp
@@ -2856,7 +2856,7 @@ static FName ConditionalGetPrefixedFormat(FName TextureFormatName, const ITarget
 	return TextureFormatName;
 }
 
-FName GetDefaultTextureFormatName( const ITargetPlatform* TargetPlatform, const UTexture* Texture, int32 LayerIndex, bool bSupportDX11TextureFormats, bool bSupportCompressedVolumeTexture, int32 BlockSize, bool bSupportFilteredFloat32Textures )
+FName GetDefaultTextureFormatName( const ITargetPlatform* TargetPlatform, const UTexture* Texture, int32 LayerIndex, bool bSupportCompressedVolumeTexture, int32 BlockSize, bool bSupportFilteredFloat32Textures )
 {
 	FName TextureFormatName = NAME_None;
 	bool bOodleTextureSdkVersionIsNone = true;
@@ -3124,23 +3124,6 @@ FName GetDefaultTextureFormatName( const ITargetPlatform* TargetPlatform, const
 	{
 		TextureFormatName = NameBGRA8;
 	}
-
-	// fallback to non-DX11 formats if one was chosen, but we can't use it
-	if (!bSupportDX11TextureFormats)
-	{
-		if (TextureFormatName == NameBC6H)
-		{
-			UE_LOG(LogTexture, Display, TEXT("BC6H Texture expanded to uncompressed RGBA16F because !SupportDX11 : %s"), *Texture->GetPathName());
-
-			TextureFormatName = NameRGBA16F;
-		}
-		else if (TextureFormatName == NameBC7)
-		{
-			UE_LOG(LogTexture, Display, TEXT("BC7 Texture expanded to uncompressed BGRA8 because !SupportDX11 : %s"), *Texture->GetPathName());
-
-			TextureFormatName = NameBGRA8;
-		}
-	}
 	
 	// remap 32F to 16F if not supported :
 	if ( !bSupportFilteredFloat32Textures &&
@@ -3176,18 +3159,18 @@ FName GetDefaultTextureFormatName( const ITargetPlatform* TargetPlatform, const
 	return ConditionalGetPrefixedFormat(TextureFormatName, TargetPlatform, bOodleTextureSdkVersionIsNone);
 }
 
-void GetDefaultTextureFormatNamePerLayer(TArray<FName>& OutFormatNames, const class ITargetPlatform* TargetPlatform, const class UTexture* Texture, bool bSupportDX11TextureFormats, bool bSupportCompressedVolumeTexture, int32 BlockSize, bool bSupportFilteredFloat32Textures )
+void GetDefaultTextureFormatNamePerLayer(TArray<FName>& OutFormatNames, const class ITargetPlatform* TargetPlatform, const class UTexture* Texture, bool bSupportCompressedVolumeTexture, int32 BlockSize, bool bSupportFilteredFloat32Textures )
 {
 #if WITH_EDITOR
 	OutFormatNames.Reserve(Texture->Source.GetNumLayers());
 	for (int32 LayerIndex = 0; LayerIndex < Texture->Source.GetNumLayers(); ++LayerIndex)
 	{
-		OutFormatNames.Add(GetDefaultTextureFormatName(TargetPlatform, Texture, LayerIndex, bSupportDX11TextureFormats, bSupportCompressedVolumeTexture, BlockSize, bSupportFilteredFloat32Textures));
+		OutFormatNames.Add(GetDefaultTextureFormatName(TargetPlatform, Texture, LayerIndex, bSupportCompressedVolumeTexture, BlockSize, bSupportFilteredFloat32Textures));
 	}
 #endif // WITH_EDITOR
 }
 
-void GetAllDefaultTextureFormats(const class ITargetPlatform* TargetPlatform, TArray<FName>& OutFormats, bool bSupportDX11TextureFormats)
+void GetAllDefaultTextureFormats(const class ITargetPlatform* TargetPlatform, TArray<FName>& OutFormats)
 {
 	// this is only used by CookOnTheFlyServer, it could be removed entirely
 
@@ -3227,11 +3210,8 @@ void GetAllDefaultTextureFormats(const class ITargetPlatform* TargetPlatform, TA
 	OutFormats.Add(NameRGBA32F);
 	OutFormats.Add(NameR16F);
 	OutFormats.Add(NameR32F);
-	if (bSupportDX11TextureFormats)
-	{
-		OutFormats.Add(NameBC6H);
-		OutFormats.Add(NameBC7);
-	}
+	OutFormats.Add(NameBC6H);
+	OutFormats.Add(NameBC7);
 	// is there any drawback to just adding the 32F textures here even if we don't want them? -> no.
 	//	what is this list even used for?
 	//  AFAICT it's only used by CookOnTheFlyServer for GetVersionFormatNumbersForIniVersionStrings
diff --git a/Engine/Source/Runtime/Engine/Public/TextureResource.h b/Engine/Source/Runtime/Engine/Public/TextureResource.h
--- a/Engine/Source/Runtime/Engine/Public/TextureResource.h
+++ b/Engine/Source/Runtime/Engine/Public/TextureResource.h
@@ -605,10 +605,10 @@ private:
 };
 
 /** Gets the name of a format for the given LayerIndex */
-ENGINE_API FName GetDefaultTextureFormatName( const class ITargetPlatform* TargetPlatform, const class UTexture* Texture, int32 LayerIndex, bool bSupportDX11TextureFormats, bool bSupportCompressedVolumeTexture = false, int32 BlockSize = 4, bool bSupportFilteredFloat32Textures = true);
+ENGINE_API FName GetDefaultTextureFormatName( const class ITargetPlatform* TargetPlatform, const class UTexture* Texture, int32 LayerIndex, bool bSupportCompressedVolumeTexture = false, int32 BlockSize = 4, bool bSupportFilteredFloat32Textures = true);
 
 /** Gets an array of format names for each layer in the texture */
-ENGINE_API void GetDefaultTextureFormatNamePerLayer(TArray<FName>& OutFormatNames, const class ITargetPlatform* TargetPlatform, const class UTexture* Texture, bool bSupportDX11TextureFormats, bool bSupportCompressedVolumeTexture = false, int32 BlockSize = 4, bool bSupportFilteredFloat32Textures = true);
+ENGINE_API void GetDefaultTextureFormatNamePerLayer(TArray<FName>& OutFormatNames, const class ITargetPlatform* TargetPlatform, const class UTexture* Texture, bool bSupportCompressedVolumeTexture = false, int32 BlockSize = 4, bool bSupportFilteredFloat32Textures = true);
 
 // returns all the texture formats which can be returned by GetDefaultTextureFormatName
-ENGINE_API void GetAllDefaultTextureFormats( const class ITargetPlatform* TargetPlatform, TArray<FName>& OutFormats, bool bSupportDX11TextureFormats);
+ENGINE_API void GetAllDefaultTextureFormats( const class ITargetPlatform* TargetPlatform, TArray<FName>& OutFormats);
