Correctly query r.Mobile.ShadingPath in a platform settings

diff --git a/Engine/Source/Developer/Android/AndroidTargetPlatform/Private/AndroidTargetPlatform.cpp b/Engine/Source/Developer/Android/AndroidTargetPlatform/Private/AndroidTargetPlatform.cpp
--- a/Engine/Source/Developer/Android/AndroidTargetPlatform/Private/AndroidTargetPlatform.cpp
+++ b/Engine/Source/Developer/Android/AndroidTargetPlatform/Private/AndroidTargetPlatform.cpp
@@ -254,6 +254,7 @@ static bool HasLicense()
 FAndroidTargetPlatform::FAndroidTargetPlatform(bool bInIsClient, const TCHAR* FlavorName, const TCHAR* OverrideIniPlatformName)
 	: TNonDesktopTargetPlatformBase(bInIsClient, FlavorName, OverrideIniPlatformName)
 	, DeviceDetection(nullptr)
+	, MobileShadingPath(0)
 	, bDistanceField(false)
 
 {
@@ -261,6 +262,7 @@ FAndroidTargetPlatform::FAndroidTargetPlatform(bool bInIsClient, const TCHAR* Fl
 	TextureLODSettings = nullptr; // These are registered by the device profile system.
 	StaticMeshLODSettings.Initialize(this);
 	GetConfigSystem()->GetBool(TEXT("/Script/Engine.RendererSettings"), TEXT("r.DistanceFields"), bDistanceField, GEngineIni);
+	GetConfigSystem()->GetInt(TEXT("/Script/Engine.RendererSettings"), TEXT("r.Mobile.ShadingPath"), MobileShadingPath, GEngineIni);
 #endif
 
 	TickDelegate = FTickerDelegate::CreateRaw(this, &FAndroidTargetPlatform::HandleTicker);
@@ -660,8 +662,7 @@ void FAndroid_ETC2TargetPlatform::GetTextureFormats(const UTexture* Texture, TAr
 
 void FAndroidTargetPlatform::GetReflectionCaptureFormats( TArray<FName>& OutFormats ) const
 {
-	static auto* MobileShadingPathCvar = IConsoleManager::Get().FindTConsoleVariableDataInt(TEXT("r.Mobile.ShadingPath"));
-	const bool bMobileDeferredShading = (MobileShadingPathCvar->GetValueOnAnyThread() == 1);
+	const bool bMobileDeferredShading = (MobileShadingPath == 1);
 	
 	if (SupportsVulkanSM5() || bMobileDeferredShading)
 	{
diff --git a/Engine/Source/Developer/Android/AndroidTargetPlatform/Private/AndroidTargetPlatform.h b/Engine/Source/Developer/Android/AndroidTargetPlatform/Private/AndroidTargetPlatform.h
--- a/Engine/Source/Developer/Android/AndroidTargetPlatform/Private/AndroidTargetPlatform.h
+++ b/Engine/Source/Developer/Android/AndroidTargetPlatform/Private/AndroidTargetPlatform.h
@@ -200,6 +200,9 @@ protected:
 
 	// Pointer to the device detection handler that grabs device ids in another thread
 	IAndroidDeviceDetection* DeviceDetection;
+
+	// r.Mobile.ShadingPath value
+	int32 MobileShadingPath;
 	
 	// true if DistanceField is enabled
 	bool bDistanceField;
diff --git a/Engine/Source/Developer/IOS/IOSTargetPlatform/Private/IOSTargetPlatform.cpp b/Engine/Source/Developer/IOS/IOSTargetPlatform/Private/IOSTargetPlatform.cpp
--- a/Engine/Source/Developer/IOS/IOSTargetPlatform/Private/IOSTargetPlatform.cpp
+++ b/Engine/Source/Developer/IOS/IOSTargetPlatform/Private/IOSTargetPlatform.cpp
@@ -31,12 +31,14 @@ FIOSTargetPlatform::FIOSTargetPlatform(bool bInIsTVOS, bool bIsClientOnly)
 	// override the ini name up in the base classes, which will go into the FTargetPlatformInfo
 	: TNonDesktopTargetPlatformBase(bIsClientOnly, nullptr, bInIsTVOS ? TEXT("TVOS") : nullptr)
 	, bIsTVOS(bInIsTVOS)
+	, MobileShadingPath(0)
 	, bDistanceField(false)
 {
 #if WITH_ENGINE
 	TextureLODSettings = nullptr; // TextureLODSettings are registered by the device profile.
 	StaticMeshLODSettings.Initialize(this);
 	GetConfigSystem()->GetBool(TEXT("/Script/Engine.RendererSettings"), TEXT("r.DistanceFields"), bDistanceField, GEngineIni);
+	GetConfigSystem()->GetInt(TEXT("/Script/Engine.RendererSettings"), TEXT("r.Mobile.ShadingPath"), MobileShadingPath, GEngineIni);
 #endif // #if WITH_ENGINE
 
 	// initialize the connected device detector
@@ -532,8 +534,7 @@ void FIOSTargetPlatform::GetAllTargetedShaderFormats( TArray<FName>& OutFormats
 
 void FIOSTargetPlatform::GetReflectionCaptureFormats( TArray<FName>& OutFormats ) const
 {
-	static auto* MobileShadingPathCvar = IConsoleManager::Get().FindTConsoleVariableDataInt(TEXT("r.Mobile.ShadingPath"));
-	const bool bMobileDeferredShading = (MobileShadingPathCvar->GetValueOnAnyThread() == 1);
+	const bool bMobileDeferredShading = (MobileShadingPath == 1);
 
 	if (SupportsMetalMRT() || bMobileDeferredShading)
 	{
diff --git a/Engine/Source/Developer/IOS/IOSTargetPlatform/Private/IOSTargetPlatform.h b/Engine/Source/Developer/IOS/IOSTargetPlatform/Private/IOSTargetPlatform.h
--- a/Engine/Source/Developer/IOS/IOSTargetPlatform/Private/IOSTargetPlatform.h
+++ b/Engine/Source/Developer/IOS/IOSTargetPlatform/Private/IOSTargetPlatform.h
@@ -131,6 +131,9 @@ private:
 	// Holds the message endpoint used for communicating with the LaunchDaemon.
 	TSharedPtr<FMessageEndpoint, ESPMode::ThreadSafe> MessageEndpoint;
 
+	// r.Mobile.ShadingPath value
+	int32 MobileShadingPath;
+
 	// true if DistanceField is enabled
 	bool bDistanceField;
 
