Fixes for truncation warnings for Android modules

#android

diff --git a/Engine/Source/Developer/Android/AndroidPlatformEditor/AndroidPlatformEditor.Build.cs b/Engine/Source/Developer/Android/AndroidPlatformEditor/AndroidPlatformEditor.Build.cs
--- a/Engine/Source/Developer/Android/AndroidPlatformEditor/AndroidPlatformEditor.Build.cs
+++ b/Engine/Source/Developer/Android/AndroidPlatformEditor/AndroidPlatformEditor.Build.cs
@@ -6,6 +6,7 @@ public class AndroidPlatformEditor : ModuleRules
 {
 	public AndroidPlatformEditor(ReadOnlyTargetRules Target) : base(Target)
 	{
+		UnsafeTypeCastWarningLevel = WarningLevel.Error; 
 		BinariesSubFolder = "Android";
 
 		PrivateDependencyModuleNames.AddRange(
diff --git a/Engine/Source/Developer/Android/AndroidPlatformEditor/Private/AndroidLicenseDialog.cpp b/Engine/Source/Developer/Android/AndroidPlatformEditor/Private/AndroidLicenseDialog.cpp
--- a/Engine/Source/Developer/Android/AndroidPlatformEditor/Private/AndroidLicenseDialog.cpp
+++ b/Engine/Source/Developer/Android/AndroidPlatformEditor/Private/AndroidLicenseDialog.cpp
@@ -43,7 +43,7 @@ void SAndroidLicenseDialog::Construct(const FArguments& InArgs)
 	if (FileReader)
 	{
 		// Create buffer for file input
-		uint32 BufferSize = FileReader->TotalSize();
+		uint32 BufferSize = IntCastChecked<uint32>(FileReader->TotalSize());
 		uint8* Buffer = (uint8*)FMemory::Malloc(BufferSize);
 		FileReader->Serialize(Buffer, BufferSize);
 
@@ -83,7 +83,7 @@ void SAndroidLicenseDialog::Construct(const FArguments& InArgs)
 
 			if (LicenseEnd < BufferEnd)
 			{
-				int32 LicenseLength = LicenseEnd - LicenseStart;
+				int32 LicenseLength = IntCastChecked<int32>(LicenseEnd - LicenseStart);
 
 				{
 					const FUTF8ToTCHAR ConvertedString(reinterpret_cast<ANSICHAR*>(LicenseStart), LicenseLength);
diff --git a/Engine/Source/Developer/Android/AndroidTargetPlatform/AndroidTargetPlatform.Build.cs b/Engine/Source/Developer/Android/AndroidTargetPlatform/AndroidTargetPlatform.Build.cs
--- a/Engine/Source/Developer/Android/AndroidTargetPlatform/AndroidTargetPlatform.Build.cs
+++ b/Engine/Source/Developer/Android/AndroidTargetPlatform/AndroidTargetPlatform.Build.cs
@@ -6,6 +6,7 @@ public class AndroidTargetPlatform : ModuleRules
 {
 	public AndroidTargetPlatform(ReadOnlyTargetRules Target) : base(Target)
 	{
+		UnsafeTypeCastWarningLevel = WarningLevel.Error;
 		BinariesSubFolder = "Android";
 
         PrivateDependencyModuleNames.AddRange(
diff --git a/Engine/Source/Developer/Android/AndroidTargetPlatform/Private/AndroidTargetPlatform.cpp b/Engine/Source/Developer/Android/AndroidTargetPlatform/Private/AndroidTargetPlatform.cpp
--- a/Engine/Source/Developer/Android/AndroidTargetPlatform/Private/AndroidTargetPlatform.cpp
+++ b/Engine/Source/Developer/Android/AndroidTargetPlatform/Private/AndroidTargetPlatform.cpp
@@ -153,7 +153,7 @@ static bool GetLicenseHash(FSHAHash& LicenseHash)
 	if (FileReader)
 	{
 		// Create buffer for file input
-		uint32 BufferSize = FileReader->TotalSize();
+		uint32 BufferSize = IntCastChecked<uint32>(FileReader->TotalSize());
 		uint8* Buffer = (uint8*)FMemory::Malloc(BufferSize);
 		FileReader->Serialize(Buffer, BufferSize);
 
@@ -191,7 +191,7 @@ static bool GetLicenseHash(FSHAHash& LicenseHash)
 
 			if (LicenseEnd < BufferEnd)
 			{
-				int32 LicenseLength = LicenseEnd - LicenseStart;
+				int32 LicenseLength = IntCastChecked<int32>(LicenseEnd - LicenseStart);
 				FSHA1::HashBuffer(LicenseStart, LicenseLength, LicenseHash.Hash);
 				bLicenseValid = true;
 			}
