Fixed: Runtime VT looks darker than streaming VT on Android
Use ETC2 RG11 compression for SVT normal maps on Android and iOS so it matches RVT format

diff --git a/Engine/Shaders/Private/VirtualTextureCompress.usf b/Engine/Shaders/Private/VirtualTextureCompress.usf
--- a/Engine/Shaders/Private/VirtualTextureCompress.usf
+++ b/Engine/Shaders/Private/VirtualTextureCompress.usf
@@ -134,7 +134,7 @@ void CompressBaseColorNormalSpecularCS(
 		BlockSpecular[i].z = round(BlockSpecular[i].z);
 	}
 #if ETC_PROFILE
-	OutCompressTexture0_128bit[ThreadId.xy + DestPos0] = CompressBlock_ETC2_RGBA(BlockBaseColor, BlockNormalX);
+	OutCompressTexture0_128bit[ThreadId.xy + DestPos0] = CompressBlock_ETC2_SRGBA(BlockBaseColor, BlockNormalX);
 	OutCompressTexture1_128bit[ThreadId.xy + DestPos1] = CompressBlock_ETC2_RGBA(BlockSpecular, BlockNormalY);
 #else
 	OutCompressTexture0_128bit[ThreadId.xy + DestPos0] = CompressBC3Block_SRGB(BlockBaseColor, BlockNormalX);
diff --git a/Engine/Source/Developer/Android/AndroidTargetPlatform/Private/AndroidTargetPlatform.cpp b/Engine/Source/Developer/Android/AndroidTargetPlatform/Private/AndroidTargetPlatform.cpp
--- a/Engine/Source/Developer/Android/AndroidTargetPlatform/Private/AndroidTargetPlatform.cpp
+++ b/Engine/Source/Developer/Android/AndroidTargetPlatform/Private/AndroidTargetPlatform.cpp
@@ -67,6 +67,7 @@ namespace AndroidTexFormat
 	const static FName NameETC2_RGB(TEXT("ETC2_RGB"));
 	const static FName NameETC2_RGBA(TEXT("ETC2_RGBA"));
 	const static FName NameETC2_R11(TEXT("ETC2_R11"));
+	const static FName NameETC2_RG11(TEXT("ETC2_RG11"));
 	const static FName NameAutoETC2(TEXT("AutoETC2"));
 
 	const static FName NameAutoASTC(TEXT("ASTC_RGBAuto"));
@@ -107,7 +108,7 @@ namespace AndroidTexFormat
 		{ NameDXT1,			NameETC2_RGB	},
 		{ NameDXT5,			NameETC2_RGBA	},
 		{ NameDXT5n,		NameETC2_RGB	},
-		{ NameBC5,			NameETC2_RGB	},
+		{ NameBC5,			NameETC2_RG11	},
 		{ NameBC4,			NameETC2_R11	},
 		{ NameBC6H,			NameRGBA16F		},
 		{ NameBC7,			NameETC2_RGBA	},
@@ -536,7 +537,7 @@ FName FAndroidTargetPlatform::FinalizeVirtualTextureLayerFormat(FName Format) co
 //		{ { FName(TEXT("ASTC_RGB_HDR")) },		{ NameRGBA16F } }, // ?
 		{ { FName(TEXT("ASTC_RGBAuto")) },		{ AndroidTexFormat::NameAutoETC2 } },
 		{ { FName(TEXT("ASTC_NormalAG")) },		{ AndroidTexFormat::NameETC2_RGB } },
-		{ { FName(TEXT("ASTC_NormalRG")) },		{ AndroidTexFormat::NameETC2_RGB } },
+		{ { FName(TEXT("ASTC_NormalRG")) },		{ AndroidTexFormat::NameETC2_RG11 } },
 		{ { AndroidTexFormat::NameDXT1 },		{ AndroidTexFormat::NameETC2_RGB } },
 		{ { AndroidTexFormat::NameDXT5 },		{ AndroidTexFormat::NameETC2_RGBA } },
 		{ { AndroidTexFormat::NameAutoDXT },	{ AndroidTexFormat::NameAutoETC2 } }
diff --git a/Engine/Source/Developer/IOS/IOSTargetPlatform/Private/IOSTargetPlatform.cpp b/Engine/Source/Developer/IOS/IOSTargetPlatform/Private/IOSTargetPlatform.cpp
--- a/Engine/Source/Developer/IOS/IOSTargetPlatform/Private/IOSTargetPlatform.cpp
+++ b/Engine/Source/Developer/IOS/IOSTargetPlatform/Private/IOSTargetPlatform.cpp
@@ -676,6 +676,7 @@ FName FIOSTargetPlatform::FinalizeVirtualTextureLayerFormat(FName Format) const
 	const static FName NameETC2_RGB(TEXT("ETC2_RGB"));
 	const static FName NameETC2_RGBA(TEXT("ETC2_RGBA"));
 	const static FName NameAutoETC2(TEXT("AutoETC2"));
+	const static FName NameETC2_RG11(TEXT("ETC2_RG11"));
 
 	// Remap non-ETC variants to ETC
 	const static FName ETCRemap[][2] =
@@ -686,7 +687,7 @@ FName FIOSTargetPlatform::FinalizeVirtualTextureLayerFormat(FName Format) const
 		{ { FName(TEXT("ASTC_RGBA_HQ")) },		{ NameETC2_RGBA } },
 //		{ { FName(TEXT("ASTC_RGB_HDR")) },		{ NameRGBA16F } }, // ?
 		{ { FName(TEXT("ASTC_NormalAG")) },		{ NameETC2_RGB } },
-		{ { FName(TEXT("ASTC_NormalRG")) },		{ NameETC2_RGB } },
+		{ { FName(TEXT("ASTC_NormalRG")) },		{ NameETC2_RG11 } },
 	};
 
 	for (int32 RemapIndex = 0; RemapIndex < UE_ARRAY_COUNT(ETCRemap); RemapIndex++)
diff --git a/Engine/Source/Developer/TextureFormatETC2/Private/TextureFormatETC2.cpp b/Engine/Source/Developer/TextureFormatETC2/Private/TextureFormatETC2.cpp
--- a/Engine/Source/Developer/TextureFormatETC2/Private/TextureFormatETC2.cpp
+++ b/Engine/Source/Developer/TextureFormatETC2/Private/TextureFormatETC2.cpp
@@ -56,6 +56,7 @@ class FETC2TextureBuildFunction final : public FTextureBuildFunction
 	op(ETC2_RGB) \
 	op(ETC2_RGBA) \
 	op(ETC2_R11) \
+	op(ETC2_RG11) \
 	op(AutoETC2)
 
 #define DECL_FORMAT_NAME(FormatName) static FName GTextureFormatName##FormatName = FName(TEXT(#FormatName));
@@ -207,7 +208,7 @@ public:
 		const struct FTextureBuildSettings* BuildSettings = nullptr
 	) const override
 	{
-		return 2;
+		return 3;
 	}
 
 	virtual FName GetEncoderName(FName Format) const override
@@ -244,6 +245,10 @@ public:
 		{
 			return PF_ETC2_R11_EAC;
 		}
+		else if (BuildSettings.TextureFormatName == GTextureFormatNameETC2_RG11)
+		{
+			return PF_ETC2_RG11_EAC;
+		}
 
 		UE_LOG(LogTextureFormatETC2, Fatal, TEXT("Unhandled texture format '%s' given to FTextureFormatAndroid::GetEncodedPixelFormat()"), *BuildSettings.TextureFormatName.ToString());
 		return PF_Unknown;
