Fix TargetPlatform includes for NoPCH

diff --git a/Engine/Source/Developer/Android/AndroidTargetPlatform/Private/AndroidTargetPlatform.h b/Engine/Source/Developer/Android/AndroidTargetPlatform/Private/AndroidTargetPlatform.h
--- a/Engine/Source/Developer/Android/AndroidTargetPlatform/Private/AndroidTargetPlatform.h
+++ b/Engine/Source/Developer/Android/AndroidTargetPlatform/Private/AndroidTargetPlatform.h
@@ -22,7 +22,7 @@
 #include "AndroidTargetDevice.h"
 
 #if WITH_ENGINE
-#include "Engine/TextureCube.h"
+#include "Engine/Texture.h"
 #include "Internationalization/Text.h"
 #include "StaticMeshResources.h"
 #endif // WITH_ENGINE
diff --git a/Engine/Source/Developer/IOS/IOSTargetPlatform/Private/IOSTargetPlatform.cpp b/Engine/Source/Developer/IOS/IOSTargetPlatform/Private/IOSTargetPlatform.cpp
--- a/Engine/Source/Developer/IOS/IOSTargetPlatform/Private/IOSTargetPlatform.cpp
+++ b/Engine/Source/Developer/IOS/IOSTargetPlatform/Private/IOSTargetPlatform.cpp
@@ -19,6 +19,7 @@
 #include "Windows/WindowsHWrapper.h"
 #endif
 #if WITH_ENGINE
+#include "Engine/Texture.h"
 #include "Sound/SoundWave.h"
 #include "TextureResource.h"
 #include "AudioCompressionSettings.h"
