Added support for ASTC Luminace+Alpha encoded normal maps

diff --git a/Engine/Shaders/Private/Common.ush b/Engine/Shaders/Private/Common.ush
--- a/Engine/Shaders/Private/Common.ush
+++ b/Engine/Shaders/Private/Common.ush
@@ -1381,6 +1381,8 @@ MaterialFloat4 UnpackNormalMap( MaterialFloat4 TextureSample )
 {
 	#if DXT5_NORMALMAPS
 		MaterialFloat2 NormalXY = TextureSample.ag;
+	#elif LA_NORMALMAPS
+		MaterialFloat2 NormalXY = TextureSample.ra;
 	#else
 		MaterialFloat2 NormalXY = TextureSample.rg;
 	#endif
diff --git a/Engine/Source/Developer/Android/AndroidTargetPlatform/Private/AndroidTargetPlatform.cpp b/Engine/Source/Developer/Android/AndroidTargetPlatform/Private/AndroidTargetPlatform.cpp
--- a/Engine/Source/Developer/Android/AndroidTargetPlatform/Private/AndroidTargetPlatform.cpp
+++ b/Engine/Source/Developer/Android/AndroidTargetPlatform/Private/AndroidTargetPlatform.cpp
@@ -71,6 +71,9 @@ namespace AndroidTexFormat
 	const static FName NameAutoETC2(TEXT("AutoETC2"));
 
 	const static FName NameAutoASTC(TEXT("ASTC_RGBAuto"));
+	const static FName NameASTC_NormalRG(TEXT("ASTC_NormalRG"));
+	// L+A mode suppoprted by ARM ASTC encoder 
+	const static FName NameASTC_NormalLA(TEXT("ASTC_NormalLA"));
 	
 	// Uncompressed Texture Formats
 	const static FName NameBGRA8(TEXT("BGRA8"));
@@ -97,7 +100,7 @@ namespace AndroidTexFormat
 		{ NameDXT1,			FName(TEXT("ASTC_RGB"))		},
 		{ NameDXT5,			FName(TEXT("ASTC_RGBA"))	},
 		{ NameDXT5n,		FName(TEXT("ASTC_NormalAG"))},
-		{ NameBC5,			FName(TEXT("ASTC_NormalRG"))},
+		{ NameBC5,			NameASTC_NormalRG			},
 		{ NameBC4,			NameETC2_R11				},
 		{ NameBC6H,			NameASTC_RGB_HDR			},
 		{ NameBC7,			FName(TEXT("ASTC_RGBA_HQ"))	},
@@ -539,7 +542,8 @@ FName FAndroidTargetPlatform::FinalizeVirtualTextureLayerFormat(FName Format) co
 //		{ { FName(TEXT("ASTC_RGB_HDR")) },		{ NameRGBA16F } }, // ?
 		{ { FName(TEXT("ASTC_RGBAuto")) },		{ AndroidTexFormat::NameAutoETC2 } },
 		{ { FName(TEXT("ASTC_NormalAG")) },		{ AndroidTexFormat::NameETC2_RGB } },
-		{ { FName(TEXT("ASTC_NormalRG")) },		{ AndroidTexFormat::NameETC2_RG11 } },
+		{ { AndroidTexFormat::NameASTC_NormalRG },	{ AndroidTexFormat::NameETC2_RG11 } },
+		{ { AndroidTexFormat::NameASTC_NormalLA	},	{ AndroidTexFormat::NameETC2_RG11 } },
 		{ { AndroidTexFormat::NameDXT1 },		{ AndroidTexFormat::NameETC2_RGB } },
 		{ { AndroidTexFormat::NameDXT5 },		{ AndroidTexFormat::NameETC2_RGBA } },
 		{ { AndroidTexFormat::NameAutoDXT },	{ AndroidTexFormat::NameAutoETC2 } }
@@ -571,6 +575,18 @@ void FAndroidTargetPlatform::GetAllTextureFormats(TArray<FName>& OutFormats) con
 	}
 }
 
+bool FAndroid_ASTCTargetPlatform::SupportsFeature(ETargetPlatformFeatures Feature) const
+{
+	if (Feature == ETargetPlatformFeatures::NormalmapLAEncodingMode)
+	{
+		static IConsoleVariable* CompressorCVar = IConsoleManager::Get().FindConsoleVariable(TEXT("cook.ASTCTextureCompressor"));
+		const bool bUsesARMCompressor = (CompressorCVar ? (CompressorCVar->GetInt() != 0) : false);
+		return bUsesARMCompressor;
+	}
+
+	return FAndroidTargetPlatform::SupportsFeature(Feature);
+}
+
 void FAndroid_ASTCTargetPlatform::GetAllTextureFormats(TArray<FName>& OutFormats) const
 {
 	FAndroidTargetPlatform::GetAllTextureFormats(OutFormats);
@@ -590,10 +606,19 @@ void FAndroid_ASTCTargetPlatform::GetTextureFormats(const UTexture* Texture, TAr
 {
 	FAndroidTargetPlatform::GetTextureFormats(Texture, OutFormats);
 
+	// L+A mode for normal map compression
+	const bool bSupportsNormalLA = SupportsFeature(ETargetPlatformFeatures::NormalmapLAEncodingMode);
+
 	// perform any remapping away from defaults
 	TArray<FName>& LayerFormats = OutFormats.Last();
 	for (FName& TextureFormatName : LayerFormats)
 	{
+		if (bSupportsNormalLA && TextureFormatName == AndroidTexFormat::NameBC5)
+		{
+			TextureFormatName = AndroidTexFormat::NameASTC_NormalLA;
+			continue;
+		}
+
 		for (int32 RemapIndex = 0; RemapIndex < UE_ARRAY_COUNT(AndroidTexFormat::ASTCRemap); ++RemapIndex)
 		{
 			if (TextureFormatName == AndroidTexFormat::ASTCRemap[RemapIndex][0])
@@ -674,6 +699,43 @@ void FAndroid_ETC2TargetPlatform::GetTextureFormats(const UTexture* Texture, TAr
 	}
 }
 
+void FAndroid_MultiTargetPlatform::GetTextureFormats(const UTexture* Texture, TArray< TArray<FName> >& OutFormats) const
+{
+	// Ask each platform variant to choose texture formats
+	for (ITargetPlatform* Platform : FormatTargetPlatforms)
+	{
+		TArray< TArray<FName> > PlatformFormats;
+		Platform->GetTextureFormats(Texture, PlatformFormats);
+		for (TArray<FName>& FormatPerLayer : PlatformFormats)
+		{
+			// For multiformat case we have to disable L+A normal map compression as only ASTC textures support it 
+			for (FName& TextureFormatName : FormatPerLayer)
+			{
+				if (TextureFormatName == AndroidTexFormat::NameASTC_NormalLA)
+				{
+					TextureFormatName = AndroidTexFormat::NameASTC_NormalRG;
+				}
+			}
+			
+			OutFormats.AddUnique(FormatPerLayer);
+		}
+	}
+}
+
+void FAndroid_MultiTargetPlatform::GetAllTextureFormats(TArray<FName>& OutFormats) const
+{
+	// Ask each platform variant to choose texture formats
+	for (ITargetPlatform* Platform : FormatTargetPlatforms)
+	{
+		TArray<FName> PlatformFormats;
+		Platform->GetAllTextureFormats(PlatformFormats);
+		for (FName Format : PlatformFormats)
+		{
+			OutFormats.AddUnique(Format);
+		}
+	}
+}
+
 void FAndroidTargetPlatform::GetReflectionCaptureFormats( TArray<FName>& OutFormats ) const
 {
 	const bool bMobileDeferredShading = (MobileShadingPath == 1);
diff --git a/Engine/Source/Developer/Android/AndroidTargetPlatform/Private/AndroidTargetPlatform.h b/Engine/Source/Developer/Android/AndroidTargetPlatform/Private/AndroidTargetPlatform.h
--- a/Engine/Source/Developer/Android/AndroidTargetPlatform/Private/AndroidTargetPlatform.h
+++ b/Engine/Source/Developer/Android/AndroidTargetPlatform/Private/AndroidTargetPlatform.h
@@ -273,6 +273,8 @@ public:
 	virtual void GetAllTextureFormats(TArray<FName>& OutFormats) const override;
 #endif
 
+	virtual bool SupportsFeature(ETargetPlatformFeatures Feature) const override;
+
 	virtual bool SupportedByExtensionsString(const FString& ExtensionsString, const int GLESVersion) const override
 	{
 		return ExtensionsString.Contains(TEXT("GL_KHR_texture_compression_astc_ldr"));
@@ -383,33 +385,8 @@ public:
 	}
 
 #if WITH_ENGINE
-	virtual void GetTextureFormats(const UTexture* Texture, TArray< TArray<FName> >& OutFormats) const
-	{
-		// Ask each platform variant to choose texture formats
-		for (ITargetPlatform* Platform : FormatTargetPlatforms)
-		{
-			TArray< TArray<FName> > PlatformFormats;
-			Platform->GetTextureFormats(Texture, PlatformFormats);
-			for (const TArray<FName>& FormatPerLayer : PlatformFormats)
-			{
-				OutFormats.AddUnique(FormatPerLayer);
-			}
-		}
-	}
-
-	virtual void GetAllTextureFormats(TArray<FName>& OutFormats) const override
-	{
-		// Ask each platform variant to choose texture formats
-		for (ITargetPlatform* Platform : FormatTargetPlatforms)
-		{
-			TArray<FName> PlatformFormats;
-			Platform->GetAllTextureFormats(PlatformFormats);
-			for (FName Format : PlatformFormats)
-			{
-				OutFormats.AddUnique(Format);
-			}
-		}
-	}
+	virtual void GetTextureFormats(const UTexture* Texture, TArray< TArray<FName> >& OutFormats) const;
+	virtual void GetAllTextureFormats(TArray<FName>& OutFormats) const override;
 #endif	
 
 	virtual float GetVariantPriority() const override
diff --git a/Engine/Source/Developer/IOS/IOSTargetPlatform/Private/IOSTargetPlatform.cpp b/Engine/Source/Developer/IOS/IOSTargetPlatform/Private/IOSTargetPlatform.cpp
--- a/Engine/Source/Developer/IOS/IOSTargetPlatform/Private/IOSTargetPlatform.cpp
+++ b/Engine/Source/Developer/IOS/IOSTargetPlatform/Private/IOSTargetPlatform.cpp
@@ -486,7 +486,13 @@ bool FIOSTargetPlatform::SupportsFeature( ETargetPlatformFeatures Feature ) cons
 
 		case ETargetPlatformFeatures::DistanceFieldAO:
 			return UsesDistanceFields();
-		
+
+		case ETargetPlatformFeatures::NormalmapLAEncodingMode:
+		{
+			static IConsoleVariable* CompressorCVar = IConsoleManager::Get().FindConsoleVariable(TEXT("cook.ASTCTextureCompressor"));
+			const bool bUsesARMCompressor = (CompressorCVar ? (CompressorCVar->GetInt() != 0) : false);
+			return bUsesARMCompressor;
+		}
 		default:
 			break;
 	}
@@ -494,7 +500,6 @@ bool FIOSTargetPlatform::SupportsFeature( ETargetPlatformFeatures Feature ) cons
 	return TTargetPlatformBase<FIOSPlatformProperties>::SupportsFeature(Feature);
 }
 
-
 void FIOSTargetPlatform::GetAllPossibleShaderFormats( TArray<FName>& OutFormats ) const
 {
 	static FName NAME_SF_METAL(TEXT("SF_METAL"));
@@ -559,6 +564,8 @@ void FIOSTargetPlatform::GetReflectionCaptureFormats( TArray<FName>& OutFormats
 }
 
 static const FName NameASTC_RGB_HDR(TEXT("ASTC_RGB_HDR"));
+static const FName NameBC5(TEXT("BC5"));
+static const FName NameASTC_NormalLA(TEXT("ASTC_NormalLA"));
 
 // we remap some of the defaults
 static const FName FormatRemap[] =
@@ -568,7 +575,7 @@ static const FName FormatRemap[] =
 	FName(TEXT("DXT1")),	FName(TEXT("ASTC_RGB")),
 	FName(TEXT("DXT5")),	FName(TEXT("ASTC_RGBA")),
 	FName(TEXT("DXT5n")),	FName(TEXT("ASTC_NormalAG")),
-	FName(TEXT("BC5")),		FName(TEXT("ASTC_NormalRG")),
+	NameBC5,				FName(TEXT("ASTC_NormalRG")),
 	FName(TEXT("BC4")),		FName(TEXT("ETC2_R11")),
 	FName(TEXT("BC6H")),	NameASTC_RGB_HDR,
 	FName(TEXT("BC7")),		FName(TEXT("ASTC_RGBA_HQ"))
@@ -607,9 +614,18 @@ void FIOSTargetPlatform::GetTextureFormats( const UTexture* Texture, TArray< TAr
 		GetDefaultTextureFormatNamePerLayer(TextureFormatNames, this, Texture, bSupportCompressedVolumeTexture, BlockSize, bSupportFilteredFloat32Textures);
 	}
 
+	// L+A mode for normal map compression
+	const bool bSupportsNormalLA = SupportsFeature(ETargetPlatformFeatures::NormalmapLAEncodingMode);
+
 	// include the formats we want
 	for (FName& TextureFormatName : TextureFormatNames)
 	{
+		if (bSupportsNormalLA && TextureFormatName == NameBC5)
+		{
+			TextureFormatName = NameASTC_NormalLA;
+			continue;
+		}
+		
 		for (int32 RemapIndex = 0; RemapIndex < UE_ARRAY_COUNT(FormatRemap); RemapIndex += 2)
 		{
 			if (TextureFormatName == FormatRemap[RemapIndex])
@@ -692,6 +708,7 @@ FName FIOSTargetPlatform::FinalizeVirtualTextureLayerFormat(FName Format) const
 //		{ { FName(TEXT("ASTC_RGB_HDR")) },		{ NameRGBA16F } }, // ?
 		{ { FName(TEXT("ASTC_NormalAG")) },		{ NameETC2_RGB } },
 		{ { FName(TEXT("ASTC_NormalRG")) },		{ NameETC2_RG11 } },
+		{ { FName(TEXT("ASTC_NormalLA")) },	    { NameETC2_RG11 } },
 	};
 
 	for (int32 RemapIndex = 0; RemapIndex < UE_ARRAY_COUNT(ETCRemap); RemapIndex++)
diff --git a/Engine/Source/Developer/TargetPlatform/Public/Interfaces/ITargetPlatform.h b/Engine/Source/Developer/TargetPlatform/Public/Interfaces/ITargetPlatform.h
--- a/Engine/Source/Developer/TargetPlatform/Public/Interfaces/ITargetPlatform.h
+++ b/Engine/Source/Developer/TargetPlatform/Public/Interfaces/ITargetPlatform.h
@@ -109,7 +109,10 @@ enum class ETargetPlatformFeatures
 	CookFileRegionMetadata,
 
 	/** The platform supports communication (reading and writing data) between a target a connected PC. */
-	DirectDataExchange
+	DirectDataExchange,
+
+	/** The platform supports Luminance + Alpha encoding mode for normalmaps */
+	NormalmapLAEncodingMode
 };
 
 enum class EPlatformAuthentication
diff --git a/Engine/Source/Developer/TextureFormatASTC/Private/TextureFormatASTC.cpp b/Engine/Source/Developer/TextureFormatASTC/Private/TextureFormatASTC.cpp
--- a/Engine/Source/Developer/TextureFormatASTC/Private/TextureFormatASTC.cpp
+++ b/Engine/Source/Developer/TextureFormatASTC/Private/TextureFormatASTC.cpp
@@ -50,7 +50,7 @@ static FAutoConsoleVariableRef CVarASTCCompressor(
 #endif
 
 // increment this if you change anything that will affect compression in this file
-#define BASE_ASTC_FORMAT_VERSION 47
+#define BASE_ASTC_FORMAT_VERSION 48
 
 #define MAX_QUALITY_BY_SIZE 4
 #define MAX_QUALITY_BY_SPEED 3
@@ -120,6 +120,13 @@ class FASTCTextureBuildFunction final : public FTextureBuildFunction
 	#pragma pack(pop)
 #endif
 
+static bool IsNormalMapFormat(FName TextureFormatName)
+{
+	return
+		TextureFormatName == GTextureFormatNameASTC_NormalAG || 
+		TextureFormatName == GTextureFormatNameASTC_NormalRG || 
+		TextureFormatName == GTextureFormatNameASTC_NormalLA;
+}
 
 static int32 GetDefaultCompressionBySizeValue(FCbObjectView InFormatConfigOverride)
 {
@@ -193,10 +200,7 @@ static EPixelFormat GetQualityFormat(const FTextureBuildSettings& BuildSettings)
 	const FCbObjectView& InFormatConfigOverride = BuildSettings.FormatConfigOverride;
 	int32 OverrideSizeValue= BuildSettings.CompressionQuality;
 
-	bool bIsNormalMap =
-		BuildSettings.TextureFormatName == GTextureFormatNameASTC_NormalAG || 
-		BuildSettings.TextureFormatName == GTextureFormatNameASTC_NormalRG || 
-		BuildSettings.TextureFormatName == GTextureFormatNameASTC_NormalLA;
+	bool bIsNormalMap = IsNormalMapFormat(BuildSettings.TextureFormatName);
 	bool bIsHQ = BuildSettings.TextureFormatName == GTextureFormatNameASTC_RGBA_HQ;
 	bool bHDRFormat = BuildSettings.TextureFormatName == GTextureFormatNameASTC_RGB_HDR;
 
@@ -272,10 +276,7 @@ static bool ASTCEnc_Compress(
 		FImageCore::SanitizeFloat16AndSetAlphaOpaqueForBC6H(Image);
 	}
 
-	bool bIsNormalMap =
-		BuildSettings.TextureFormatName == GTextureFormatNameASTC_NormalAG || 
-		BuildSettings.TextureFormatName == GTextureFormatNameASTC_NormalRG || 
-		BuildSettings.TextureFormatName == GTextureFormatNameASTC_NormalLA;
+	bool bIsNormalMap = IsNormalMapFormat(BuildSettings.TextureFormatName);
 		
 	// Determine the compressed pixel format and compression parameters
 	EPixelFormat CompressedPixelFormat = GetQualityFormat(BuildSettings);
diff --git a/Engine/Source/Runtime/Engine/Private/ShaderCompiler/ShaderCompiler.cpp b/Engine/Source/Runtime/Engine/Private/ShaderCompiler/ShaderCompiler.cpp
--- a/Engine/Source/Runtime/Engine/Private/ShaderCompiler/ShaderCompiler.cpp
+++ b/Engine/Source/Runtime/Engine/Private/ShaderCompiler/ShaderCompiler.cpp
@@ -6603,6 +6603,12 @@ void GlobalBeginCompileShader(
 		Input.Environment.SetDefine(TEXT("POST_PROCESS_ALPHA"), PropagateAlpha);
 	}
 
+	if (TargetPlatform && 
+		TargetPlatform->SupportsFeature(ETargetPlatformFeatures::NormalmapLAEncodingMode))
+	{
+		Input.Environment.SetDefine(TEXT("LA_NORMALMAPS"), 1);
+	}
+
 	Input.Environment.SetDefine(TEXT("PLATFORM_SUPPORTS_RENDERTARGET_WRITE_MASK"), RHISupportsRenderTargetWriteMask(EShaderPlatform(Target.Platform)) ? 1 : 0);
 	Input.Environment.SetDefine(TEXT("PLATFORM_SUPPORTS_PER_PIXEL_DBUFFER_MASK"), FDataDrivenShaderPlatformInfo::GetSupportsPerPixelDBufferMask(EShaderPlatform(Target.Platform)) ? 1 : 0);
 	Input.Environment.SetDefine(TEXT("PLATFORM_SUPPORTS_DISTANCE_FIELDS"), DoesPlatformSupportDistanceFields(EShaderPlatform(Target.Platform)) ? 1 : 0);
diff --git a/Engine/Source/Runtime/RenderCore/Private/Shader.cpp b/Engine/Source/Runtime/RenderCore/Private/Shader.cpp
--- a/Engine/Source/Runtime/RenderCore/Private/Shader.cpp
+++ b/Engine/Source/Runtime/RenderCore/Private/Shader.cpp
@@ -1879,6 +1879,12 @@ void ShaderMapAppendKeyString(EShaderPlatform Platform, FString& KeyString)
 		}
 	}
 
+	if (TargetPlatform && 
+		TargetPlatform->SupportsFeature(ETargetPlatformFeatures::NormalmapLAEncodingMode))
+	{
+		KeyString += TEXT("_NLA");
+	}
+
 	{
 		static const auto CVar = IConsoleManager::Get().FindTConsoleVariableDataInt(TEXT("r.VertexFoggingForOpaque"));
 		bool bVertexFoggingForOpaque = CVar && CVar->GetValueOnAnyThread() > 0;
