Fixed: linking issues when building program target

diff --git a/Engine/Source/Developer/Android/AndroidTargetPlatform/Private/AndroidTargetPlatform.cpp b/Engine/Source/Developer/Android/AndroidTargetPlatform/Private/AndroidTargetPlatform.cpp
--- a/Engine/Source/Developer/Android/AndroidTargetPlatform/Private/AndroidTargetPlatform.cpp
+++ b/Engine/Source/Developer/Android/AndroidTargetPlatform/Private/AndroidTargetPlatform.cpp
@@ -413,6 +413,13 @@ bool FAndroidTargetPlatform::SupportsFeature( ETargetPlatformFeatures Feature )
 
 		case ETargetPlatformFeatures::DistanceFieldAO:
 			return UsesDistanceFields();
+
+		case ETargetPlatformFeatures::NormalmapLAEncodingMode:
+		{
+			static IConsoleVariable* CompressorCVar = IConsoleManager::Get().FindConsoleVariable(TEXT("cook.ASTCTextureCompressor"));
+			const bool bUsesARMCompressor = (CompressorCVar ? (CompressorCVar->GetInt() != 0) : false);
+			return SupportsTextureFormatCategory(EAndroidTextureFormatCategory::ASTC) && bUsesARMCompressor;
+		}
 			
 		default:
 			break;
@@ -575,18 +582,6 @@ void FAndroidTargetPlatform::GetAllTextureFormats(TArray<FName>& OutFormats) con
 	}
 }
 
-bool FAndroid_ASTCTargetPlatform::SupportsFeature(ETargetPlatformFeatures Feature) const
-{
-	if (Feature == ETargetPlatformFeatures::NormalmapLAEncodingMode)
-	{
-		static IConsoleVariable* CompressorCVar = IConsoleManager::Get().FindConsoleVariable(TEXT("cook.ASTCTextureCompressor"));
-		const bool bUsesARMCompressor = (CompressorCVar ? (CompressorCVar->GetInt() != 0) : false);
-		return bUsesARMCompressor;
-	}
-
-	return FAndroidTargetPlatform::SupportsFeature(Feature);
-}
-
 void FAndroid_ASTCTargetPlatform::GetAllTextureFormats(TArray<FName>& OutFormats) const
 {
 	FAndroidTargetPlatform::GetAllTextureFormats(OutFormats);
diff --git a/Engine/Source/Developer/Android/AndroidTargetPlatform/Private/AndroidTargetPlatform.h b/Engine/Source/Developer/Android/AndroidTargetPlatform/Private/AndroidTargetPlatform.h
--- a/Engine/Source/Developer/Android/AndroidTargetPlatform/Private/AndroidTargetPlatform.h
+++ b/Engine/Source/Developer/Android/AndroidTargetPlatform/Private/AndroidTargetPlatform.h
@@ -273,8 +273,6 @@ public:
 	virtual void GetAllTextureFormats(TArray<FName>& OutFormats) const override;
 #endif
 
-	virtual bool SupportsFeature(ETargetPlatformFeatures Feature) const override;
-
 	virtual bool SupportedByExtensionsString(const FString& ExtensionsString, const int GLESVersion) const override
 	{
 		return ExtensionsString.Contains(TEXT("GL_KHR_texture_compression_astc_ldr"));
