OpenXR: fix Quest packaging telemetry

We changed the config attribute to a bool, but didn't change the analytics event API from Array to Bool.

#rnx

diff --git a/Engine/Source/Developer/Android/AndroidTargetPlatform/Private/AndroidTargetPlatform.cpp b/Engine/Source/Developer/Android/AndroidTargetPlatform/Private/AndroidTargetPlatform.cpp
--- a/Engine/Source/Developer/Android/AndroidTargetPlatform/Private/AndroidTargetPlatform.cpp
+++ b/Engine/Source/Developer/Android/AndroidTargetPlatform/Private/AndroidTargetPlatform.cpp
@@ -468,7 +468,7 @@ void FAndroidTargetPlatform::GetPlatformSpecificProjectAnalytics( TArray<FAnalyt
 		TEXT("SupportsES31"), SupportsES31()
 	);
 
-	AppendAnalyticsEventConfigArray(AnalyticsParamArray, TEXT("/Script/AndroidRuntimeSettings.AndroidRuntimeSettings"), TEXT("bPackageForMetaQuest"), GEngineIni);
+	AppendAnalyticsEventConfigBool(AnalyticsParamArray, TEXT("/Script/AndroidRuntimeSettings.AndroidRuntimeSettings"), TEXT("bPackageForMetaQuest"), GEngineIni);
 }
 
 #if WITH_ENGINE
