More operator new(TArray) fixes.

diff --git a/Engine/Source/Developer/Android/AndroidPlatformEditor/Private/AndroidTargetSettingsCustomization.cpp b/Engine/Source/Developer/Android/AndroidPlatformEditor/Private/AndroidTargetSettingsCustomization.cpp
--- a/Engine/Source/Developer/Android/AndroidPlatformEditor/Private/AndroidTargetSettingsCustomization.cpp
+++ b/Engine/Source/Developer/Android/AndroidPlatformEditor/Private/AndroidTargetSettingsCustomization.cpp
@@ -63,16 +63,16 @@ FAndroidTargetSettingsCustomization::FAndroidTargetSettingsCustomization()
 	, EngineProjectPropertiesPath(EngineAndroidPath / TEXT("project.properties"))
 	, GameProjectPropertiesPath(GameAndroidPath / TEXT("project.properties"))
 {
-	new (IconNames) FPlatformIconInfo(TEXT("res/drawable/icon.png"), LOCTEXT("SettingsIcon", "Icon"), FText::GetEmpty(), 48, 48, FPlatformIconInfo::Required);
-	new (IconNames) FPlatformIconInfo(TEXT("res/drawable-ldpi/icon.png"), LOCTEXT("SettingsIcon_LDPI", "LDPI Icon"), FText::GetEmpty(), 36, 36, FPlatformIconInfo::Required);
-	new (IconNames) FPlatformIconInfo(TEXT("res/drawable-mdpi/icon.png"), LOCTEXT("SettingsIcon_MDPI", "MDPI Icon"), FText::GetEmpty(), 48, 48, FPlatformIconInfo::Required);
-	new (IconNames) FPlatformIconInfo(TEXT("res/drawable-hdpi/icon.png"), LOCTEXT("SettingsIcon_HDPI", "HDPI Icon"), FText::GetEmpty(), 72, 72, FPlatformIconInfo::Required);
-	new (IconNames) FPlatformIconInfo(TEXT("res/drawable-xhdpi/icon.png"), LOCTEXT("SettingsIcon_XHDPI", "XHDPI Icon"), FText::GetEmpty(), 96, 96, FPlatformIconInfo::Required);
-
-	new (LaunchImageNames)FPlatformIconInfo(TEXT("res/drawable/downloadimagev.png"), LOCTEXT("SettingsIcon_DownloadImageV", "Download Background Vertical Image"), FText::GetEmpty(), 720, 1280, FPlatformIconInfo::Required);
-	new (LaunchImageNames)FPlatformIconInfo(TEXT("res/drawable/downloadimageh.png"), LOCTEXT("SettingsIcon_DownloadImageH", "Download Background Horizontal Image"), FText::GetEmpty(), 1280, 720, FPlatformIconInfo::Required);
-	new (LaunchImageNames)FPlatformIconInfo(TEXT("res/drawable/splashscreen_portrait.png"), LOCTEXT("LaunchImage_Portrait", "Launch Portrait"), FText::GetEmpty(), 360, 640, FPlatformIconInfo::Required);
-	new (LaunchImageNames)FPlatformIconInfo(TEXT("res/drawable/splashscreen_landscape.png"), LOCTEXT("LaunchImage_Landscape", "Launch Landscape"), FText::GetEmpty(), 640, 360, FPlatformIconInfo::Required);
+	IconNames.Emplace(TEXT("res/drawable/icon.png"), LOCTEXT("SettingsIcon", "Icon"), FText::GetEmpty(), 48, 48, FPlatformIconInfo::Required);
+	IconNames.Emplace(TEXT("res/drawable-ldpi/icon.png"), LOCTEXT("SettingsIcon_LDPI", "LDPI Icon"), FText::GetEmpty(), 36, 36, FPlatformIconInfo::Required);
+	IconNames.Emplace(TEXT("res/drawable-mdpi/icon.png"), LOCTEXT("SettingsIcon_MDPI", "MDPI Icon"), FText::GetEmpty(), 48, 48, FPlatformIconInfo::Required);
+	IconNames.Emplace(TEXT("res/drawable-hdpi/icon.png"), LOCTEXT("SettingsIcon_HDPI", "HDPI Icon"), FText::GetEmpty(), 72, 72, FPlatformIconInfo::Required);
+	IconNames.Emplace(TEXT("res/drawable-xhdpi/icon.png"), LOCTEXT("SettingsIcon_XHDPI", "XHDPI Icon"), FText::GetEmpty(), 96, 96, FPlatformIconInfo::Required);
+
+	LaunchImageNames.Emplace(TEXT("res/drawable/downloadimagev.png"), LOCTEXT("SettingsIcon_DownloadImageV", "Download Background Vertical Image"), FText::GetEmpty(), 720, 1280, FPlatformIconInfo::Required);
+	LaunchImageNames.Emplace(TEXT("res/drawable/downloadimageh.png"), LOCTEXT("SettingsIcon_DownloadImageH", "Download Background Horizontal Image"), FText::GetEmpty(), 1280, 720, FPlatformIconInfo::Required);
+	LaunchImageNames.Emplace(TEXT("res/drawable/splashscreen_portrait.png"), LOCTEXT("LaunchImage_Portrait", "Launch Portrait"), FText::GetEmpty(), 360, 640, FPlatformIconInfo::Required);
+	LaunchImageNames.Emplace(TEXT("res/drawable/splashscreen_landscape.png"), LOCTEXT("LaunchImage_Landscape", "Launch Landscape"), FText::GetEmpty(), 640, 360, FPlatformIconInfo::Required);
 }
 
 void FAndroidTargetSettingsCustomization::CustomizeDetails(IDetailLayoutBuilder& DetailLayout)
diff --git a/Engine/Source/Editor/UnrealEd/Private/Commandlets/InternationalizationConditioningCommandlet.cpp b/Engine/Source/Editor/UnrealEd/Private/Commandlets/InternationalizationConditioningCommandlet.cpp
--- a/Engine/Source/Editor/UnrealEd/Private/Commandlets/InternationalizationConditioningCommandlet.cpp
+++ b/Engine/Source/Editor/UnrealEd/Private/Commandlets/InternationalizationConditioningCommandlet.cpp
@@ -15,7 +15,6 @@
 #include "Misc/AssertionMacros.h"
 #include "Serialization/JsonInternationalizationArchiveSerializer.h"
 #include "Serialization/JsonInternationalizationManifestSerializer.h"
-#include "Templates/ChooseClass.h"
 #include "Trace/Detail/Channel.h"
 #include "UObject/NameTypes.h"
 
@@ -98,24 +97,24 @@ void UInternationalizationConditioningCommandlet::FLocalizationFile::CompareToCo
 					}
 					else if ( PropValue == OtherValue->GetValue() )
 					{
-						new(IdenticalProperties) FLocalizationFileEntry( Other->GetFilename(), LocSectionName, Propname.ToString(), EscapedPropValue, EscapedPropValue );
+						IdenticalProperties.Emplace( Other->GetFilename(), LocSectionName, Propname.ToString(), EscapedPropValue, EscapedPropValue );
 					}
 					else
 					{
-						new(TranslatedProperties) FLocalizationFileEntry( Other->GetFilename(), LocSectionName, Propname.ToString(), EscapedPropValue, EscapedOtherValue );
+						TranslatedProperties.Emplace( Other->GetFilename(), LocSectionName, Propname.ToString(), EscapedPropValue, EscapedOtherValue );
 					}
 				}
 				else
 				{
 					// The counterpart didn't contain this key
-					new(UnmatchedProperties) FString(LocSectionName + TEXT(".") + Propname.ToString());
+					UnmatchedProperties.Add(LocSectionName + TEXT(".") + Propname.ToString());
 				}
 			}
 		}
 		else
 		{
 			// The counterpart didn't contain this section
-			new(UnmatchedSections) FString(FPaths::GetBaseFilename(LocFilename) + TEXT(".") + LocSectionName);
+			UnmatchedSections.Add(FPaths::GetBaseFilename(LocFilename) + TEXT(".") + LocSectionName);
 		}
 	}
 }
@@ -322,7 +321,7 @@ bool UInternationalizationConditioningCommandlet::ProcessManifest( const FString
 	IFileManager::Get().FindFiles(PathPrimaryFilenames, *PrimaryWildcardName, true, false);
 	for ( int32 FileIndex = 0; FileIndex < PathPrimaryFilenames.Num(); FileIndex++ )
 	{
-		FString* CompleteFilename = new(PrimaryFilenames) FString(PrimaryLocDirectory + PathPrimaryFilenames[FileIndex]);
+		PrimaryFilenames.Add(PrimaryLocDirectory + PathPrimaryFilenames[FileIndex]);
 	}
 
 	if ( PrimaryFilenames.Num() == 0 )
@@ -413,7 +412,7 @@ bool UInternationalizationConditioningCommandlet::ProcessArchive( const FString&
 	IFileManager::Get().FindFiles(PathPrimaryFilenames, *PrimaryWildcardName, true, false);
 	for ( int32 FileIndex = 0; FileIndex < PathPrimaryFilenames.Num(); FileIndex++ )
 	{
-		FString* CompleteFilename = new(PrimaryFilenames) FString(PrimaryLocDirectory + PathPrimaryFilenames[FileIndex]);
+		PrimaryFilenames.Add(PrimaryLocDirectory + PathPrimaryFilenames[FileIndex]);
 	}
 
 	if ( PrimaryFilenames.Num() == 0 )
@@ -436,7 +435,7 @@ bool UInternationalizationConditioningCommandlet::ProcessArchive( const FString&
 
 		for ( int32 FileIndex = 0; FileIndex < PathForeignFilenames.Num(); FileIndex++ )
 		{
-			FString* CompleteFilename = new(ForeignFilenames) FString(ForeignLocDirectory + PathForeignFilenames[FileIndex]);
+			ForeignFilenames.Add(ForeignLocDirectory + PathForeignFilenames[FileIndex]);
 		}
 
 		if ( ForeignFilenames.Num() == 0 )
diff --git a/Engine/Source/Editor/UnrealEd/Private/Commandlets/PackageUtilities.cpp b/Engine/Source/Editor/UnrealEd/Private/Commandlets/PackageUtilities.cpp
--- a/Engine/Source/Editor/UnrealEd/Private/Commandlets/PackageUtilities.cpp
+++ b/Engine/Source/Editor/UnrealEd/Private/Commandlets/PackageUtilities.cpp
@@ -73,7 +73,7 @@ void SearchDirectoryRecursive( const FString& SearchPathMask, TArray<FString>& o
 	{
 		for ( int32 PkgIndex = 0; PkgIndex < PackageNames.Num(); PkgIndex++ )
 		{
-			new(out_PackageFilenames) FString( SearchPath / PackageNames[PkgIndex] );
+			out_PackageFilenames.Add( SearchPath / PackageNames[PkgIndex] );
 		}
 
 		out_PackageNames += PackageNames;
@@ -1138,7 +1138,7 @@ void FPkgInfoReporter_Log::GeneratePackageReport( FLinkerLoad* InLinker /*=nullp
 		SortedExportMap.Empty(Linker->ExportMap.Num());
 		for( int32 i = 0; i < Linker->ExportMap.Num(); ++i )
 		{
-			new(SortedExportMap) FExportInfo(Linker, i);
+			SortedExportMap.Emplace(Linker, i);
 		}
 
 		FString SortingParms;
diff --git a/Engine/Source/Editor/UnrealEd/Private/ReferenceInfoUtils.cpp b/Engine/Source/Editor/UnrealEd/Private/ReferenceInfoUtils.cpp
--- a/Engine/Source/Editor/UnrealEd/Private/ReferenceInfoUtils.cpp
+++ b/Engine/Source/Editor/UnrealEd/Private/ReferenceInfoUtils.cpp
@@ -386,10 +386,10 @@ namespace ReferenceInfoUtils
 		// If any BSP surfaces are selected
 		if (BspMats.Num() > 0)
 		{
-			FReferencedAssets* Referencer = new(Referencers) FReferencedAssets(InWorld->GetModel());
+			FReferencedAssets& Referencer = Referencers.Emplace_GetRef(InWorld->GetModel());
 
 			// Now copy the array
-			Referencer->AssetList = BspMats;
+			Referencer.AssetList = BspMats;
 			ReferenceGraph.Add(ObjectPtrWrap(InWorld->GetModel()), ObjectPtrWrap(BspMats));
 		}
 
@@ -424,10 +424,10 @@ namespace ReferenceInfoUtils
 			SelectedActors[Index]->Mark(OBJECTMARK_TagExp);
 
 			// Create a new entry for this actor
-			FReferencedAssets* Referencer = new(Referencers) FReferencedAssets(SelectedActors[Index]);
+			FReferencedAssets& Referencer = Referencers.Emplace_GetRef(SelectedActors[Index]);
 
 			// Add to the list of referenced assets
-			FFindAssetsArchive(SelectedActors[Index], Referencer->AssetList, &ReferenceGraph, MaxRecursionDepth, bShowScript, bShowDefault);
+			FFindAssetsArchive(SelectedActors[Index], Referencer.AssetList, &ReferenceGraph, MaxRecursionDepth, bShowScript, bShowDefault);
 		}
 
 		// Rebuild the name cache
diff --git a/Engine/Source/Editor/UnrealEd/Private/ReferencedAssetsUtils.cpp b/Engine/Source/Editor/UnrealEd/Private/ReferencedAssetsUtils.cpp
--- a/Engine/Source/Editor/UnrealEd/Private/ReferencedAssetsUtils.cpp
+++ b/Engine/Source/Editor/UnrealEd/Private/ReferencedAssetsUtils.cpp
@@ -143,7 +143,7 @@ void FFindReferencedAssets::BuildAssetList(UObject *Object, const TArray<UClass*
 	TArray<FReferencedAssets> LocalReferencers;
 
 	// Create a new entry for this actor.
-	new( LocalReferencers ) FReferencedAssets( Object );
+	LocalReferencers.Emplace( Object );
 
 	for (FThreadSafeObjectIterator It; It; ++It)
 	{
diff --git a/Engine/Source/Runtime/Engine/Private/Particles/ParticleCurveTexture.cpp b/Engine/Source/Runtime/Engine/Private/Particles/ParticleCurveTexture.cpp
--- a/Engine/Source/Runtime/Engine/Private/Particles/ParticleCurveTexture.cpp
+++ b/Engine/Source/Runtime/Engine/Private/Particles/ParticleCurveTexture.cpp
@@ -548,10 +548,10 @@ FTexelAllocation FParticleCurveTexture::AddCurve(const TArray<FColor>& CurveSamp
 			if (TexelAllocation.Size > 0)
 			{
 				check(TexelAllocation.Size == CurveSamples.Num());
-				FCurveSamples* PendingCurve = new(PendingCurves) FCurveSamples;
-				PendingCurve->TexelAllocation = TexelAllocation;
-				PendingCurve->Samples = (FColor*)FMemory::Malloc(TexelAllocation.Size * sizeof(FColor));
-				FMemory::Memcpy(PendingCurve->Samples, CurveSamples.GetData(), TexelAllocation.Size * sizeof(FColor));
+				FCurveSamples& PendingCurve = PendingCurves.AddDefaulted_GetRef();
+				PendingCurve.TexelAllocation = TexelAllocation;
+				PendingCurve.Samples = (FColor*)FMemory::Malloc(TexelAllocation.Size * sizeof(FColor));
+				FMemory::Memcpy(PendingCurve.Samples, CurveSamples.GetData(), TexelAllocation.Size * sizeof(FColor));
 				return TexelAllocation;
 			}
 			UE_LOG(LogParticles, Warning, TEXT("FParticleCurveTexture: Failed to allocate %d texels for a curve (may need to increase the size of GParticleCurveTextureSizeX or GParticleCurveTextureSizeY)."), CurveSamples.Num());
diff --git a/Engine/Source/Runtime/Engine/Private/Particles/ParticleGpuSimulation.cpp b/Engine/Source/Runtime/Engine/Private/Particles/ParticleGpuSimulation.cpp
--- a/Engine/Source/Runtime/Engine/Private/Particles/ParticleGpuSimulation.cpp
+++ b/Engine/Source/Runtime/Engine/Private/Particles/ParticleGpuSimulation.cpp
@@ -3777,7 +3777,7 @@ private:
 				TileOffset.Y = FMath::Fractional(FMath::TruncToFloat((float)TileIndex / (float)GParticleSimulationTileCountX) / (float)GParticleSimulationTileCountY);
 				FreeParticlesInTile = GParticlesPerTile;
 			}
-			FNewParticle& Particle = *new(InNewParticles) FNewParticle();
+			FNewParticle& Particle = InNewParticles.AddDefaulted_GetRef();
 			const int32 SubTileIndex = GParticlesPerTile - FreeParticlesInTile;
 			const int32 SubTileX = SubTileIndex % GParticleSimulationTileSize;
 			const int32 SubTileY = SubTileIndex / GParticleSimulationTileSize;
@@ -4560,7 +4560,7 @@ void FFXSystem::SimulateGPUParticles(
 		check(Simulation);
 		if (Simulation->SimulationPhase == Phase && Simulation->TileVertexBuffer.TileCount > 0 && Simulation->bEnabled)
 		{
-			FSimulationCommandGPU* SimulationCommand = new(SimulationCommands) FSimulationCommandGPU(
+			FSimulationCommandGPU& SimulationCommand = SimulationCommands.Emplace_GetRef(
 				Simulation->TileVertexBuffer.GetShaderParam(),
 				Simulation->TileVertexBuffer.VertexBufferRHI,
 				Simulation->EmitterSimulationResources->SimulationUniformBuffer,
@@ -4585,7 +4585,7 @@ void FFXSystem::SimulateGPUParticles(
 					if (FMath::Abs(LocalIntensity) > 0.0f)
 					{
 						Simulation->LocalVectorField.Resource->Update(RHICmdList, Simulation->PerFrameSimulationParameters.DeltaSeconds);
-						SimulationCommand->VectorFieldTexturesRHI[0] = Simulation->LocalVectorField.Resource->VolumeTextureRHI;
+						SimulationCommand.VectorFieldTexturesRHI[0] = Simulation->LocalVectorField.Resource->VolumeTextureRHI;
 						SetParametersForVectorField(VectorFieldParameters, &Simulation->LocalVectorField, /*EmitterScale=*/ 1.0f, /*EmitterTightness=*/ -1, VectorFieldParameters.Count++);
 					}
 				}
@@ -4604,7 +4604,7 @@ void FFXSystem::SimulateGPUParticles(
 						if (SimulationBounds.Intersect(Instance->WorldBounds) &&
 							FMath::Abs(Intensity) > 0.0f)
 						{
-							SimulationCommand->VectorFieldTexturesRHI[VectorFieldParameters.Count] = Instance->Resource->VolumeTextureRHI;
+							SimulationCommand.VectorFieldTexturesRHI[VectorFieldParameters.Count] = Instance->Resource->VolumeTextureRHI;
 							SetParametersForVectorField(VectorFieldParameters, Instance, GlobalVectorFieldScale, GlobalVectorFieldTightness, VectorFieldParameters.Count++);
 						}
 					}
@@ -4634,9 +4634,9 @@ void FFXSystem::SimulateGPUParticles(
 						Simulation->LocalVectorFieldUniformBuffer = FVectorFieldUniformBufferRef::CreateUniformBufferImmediate(VectorFieldParameters, UniformBuffer_MultiFrame);
 						Simulation->LocalIntensity = LocalIntensity;
 					}
-					SimulationCommand->VectorFieldsUniformBuffer = Simulation->LocalVectorFieldUniformBuffer;
+					SimulationCommand.VectorFieldsUniformBuffer = Simulation->LocalVectorFieldUniformBuffer;
 #else
-					SimulationCommand->VectorFieldsUniformBuffer = FVectorFieldUniformBufferRef::CreateUniformBufferImmediate(VectorFieldParameters, UniformBuffer_SingleFrame);
+					SimulationCommand.VectorFieldsUniformBuffer = FVectorFieldUniformBufferRef::CreateUniformBufferImmediate(VectorFieldParameters, UniformBuffer_SingleFrame);
 #endif
 				}
 			}
diff --git a/Engine/Source/Runtime/VulkanRHI/Private/VulkanMemory.cpp b/Engine/Source/Runtime/VulkanRHI/Private/VulkanMemory.cpp
--- a/Engine/Source/Runtime/VulkanRHI/Private/VulkanMemory.cpp
+++ b/Engine/Source/Runtime/VulkanRHI/Private/VulkanMemory.cpp
@@ -4369,9 +4369,9 @@ namespace VulkanRHI
 			}
 		}
 
-		FPendingItemsPerCmdBuffer* New = new(PendingFreeStagingBuffers) FPendingItemsPerCmdBuffer;
-		New->CmdBuffer = CmdBuffer;
-		return New;
+		FPendingItemsPerCmdBuffer& New = PendingFreeStagingBuffers.AddDefaulted_GetRef();
+		New.CmdBuffer = CmdBuffer;
+		return &New;
 	}
 
 	inline FStagingManager::FPendingItemsPerCmdBuffer::FPendingItems* FStagingManager::FPendingItemsPerCmdBuffer::FindOrAddItemsForFence(uint64 Fence)
@@ -4384,9 +4384,9 @@ namespace VulkanRHI
 			}
 		}
 
-		FPendingItems* New = new(PendingItems) FPendingItems;
-		New->FenceCounter = Fence;
-		return New;
+		FPendingItems& New = PendingItems.AddDefaulted_GetRef();
+		New.FenceCounter = Fence;
+		return &New;
 	}
 
 	void FStagingManager::ReleaseBuffer(FVulkanCmdBuffer* CmdBuffer, FStagingBuffer*& StagingBuffer)
diff --git a/Engine/Source/Runtime/VulkanRHI/Private/VulkanPendingState.cpp b/Engine/Source/Runtime/VulkanRHI/Private/VulkanPendingState.cpp
--- a/Engine/Source/Runtime/VulkanRHI/Private/VulkanPendingState.cpp
+++ b/Engine/Source/Runtime/VulkanRHI/Private/VulkanPendingState.cpp
@@ -31,10 +31,10 @@ FVulkanDescriptorPool::FVulkanDescriptorPool(FVulkanDevice* InDevice, const FVul
 		uint32 NumTypesUsed = Layout.GetTypesUsed(DescriptorType);
 		if (NumTypesUsed > 0)
 		{
-			VkDescriptorPoolSize* Type = new(Types) VkDescriptorPoolSize;
-			FMemory::Memzero(*Type);
-			Type->type = DescriptorType;
-			Type->descriptorCount = NumTypesUsed * MaxSetsAllocations;
+			VkDescriptorPoolSize& Type = Types.AddDefaulted_GetRef();
+			FMemory::Memzero(Type);
+			Type.type = DescriptorType;
+			Type.descriptorCount = NumTypesUsed * MaxSetsAllocations;
 		}
 	}
 
@@ -43,10 +43,10 @@ FVulkanDescriptorPool::FVulkanDescriptorPool(FVulkanDevice* InDevice, const FVul
 		uint32 NumTypesUsed = Layout.GetTypesUsed(VK_DESCRIPTOR_TYPE_ACCELERATION_STRUCTURE_KHR);
 		if (NumTypesUsed > 0)
 		{
-			VkDescriptorPoolSize* Type = new(Types) VkDescriptorPoolSize;
-			FMemory::Memzero(*Type);
-			Type->type = VK_DESCRIPTOR_TYPE_ACCELERATION_STRUCTURE_KHR;
-			Type->descriptorCount = NumTypesUsed * MaxSetsAllocations;
+			VkDescriptorPoolSize& Type = Types.AddDefaulted_GetRef();
+			FMemory::Memzero(Type);
+			Type.type = VK_DESCRIPTOR_TYPE_ACCELERATION_STRUCTURE_KHR;
+			Type.descriptorCount = NumTypesUsed * MaxSetsAllocations;
 		}
 	}
 #endif
diff --git a/Engine/Source/Runtime/VulkanRHI/Private/VulkanPipeline.cpp b/Engine/Source/Runtime/VulkanRHI/Private/VulkanPipeline.cpp
--- a/Engine/Source/Runtime/VulkanRHI/Private/VulkanPipeline.cpp
+++ b/Engine/Source/Runtime/VulkanRHI/Private/VulkanPipeline.cpp
@@ -1019,8 +1019,8 @@ void FGfxPipelineDesc::FRenderTargets::ReadFrom(const FVulkanRenderTargetLayout&
 	{
 		for (uint32 Index = 0; Index < Count; ++Index)
 		{
-			FGfxPipelineDesc::FRenderTargets::FAttachmentRef* New = new(Dest) FGfxPipelineDesc::FRenderTargets::FAttachmentRef;
-			New->ReadFrom(Source[Index]);
+			FGfxPipelineDesc::FRenderTargets::FAttachmentRef& New = Dest.AddDefaulted_GetRef();
+			New.ReadFrom(Source[Index]);
 		}
 	};
 	CopyAttachmentRefs(ColorAttachments, RTLayout.ColorReferences, UE_ARRAY_COUNT(RTLayout.ColorReferences));
@@ -1303,16 +1303,16 @@ bool FVulkanPipelineStateCacheManager::CreateGfxPipelineFromEntry(FVulkanRHIGrap
 	TArray<VkVertexInputBindingDescription, TInlineAllocator<32>> VBBindings;
 	for (const FGfxPipelineDesc::FVertexBinding& SourceBinding : GfxEntry->VertexBindings)
 	{
-		VkVertexInputBindingDescription* Binding = new(VBBindings) VkVertexInputBindingDescription;
-		SourceBinding.WriteInto(*Binding);
+		VkVertexInputBindingDescription& Binding = VBBindings.AddDefaulted_GetRef();
+		SourceBinding.WriteInto(Binding);
 	}
 	VBInfo.vertexBindingDescriptionCount = VBBindings.Num();
 	VBInfo.pVertexBindingDescriptions = VBBindings.GetData();
 	TArray<VkVertexInputAttributeDescription, TInlineAllocator<32>> VBAttributes;
 	for (const FGfxPipelineDesc::FVertexAttribute& SourceAttr : GfxEntry->VertexAttributes)
 	{
-		VkVertexInputAttributeDescription* Attr = new(VBAttributes) VkVertexInputAttributeDescription;
-		SourceAttr.WriteInto(*Attr);
+		VkVertexInputAttributeDescription& Attr = VBAttributes.AddDefaulted_GetRef();
+		SourceAttr.WriteInto(Attr);
 	}
 	VBInfo.vertexAttributeDescriptionCount = VBAttributes.Num();
 	VBInfo.pVertexAttributeDescriptions = VBAttributes.GetData();
@@ -1759,8 +1759,8 @@ void FVulkanPipelineStateCacheManager::CreateGfxEntry(const FGraphicsPipelineSta
 	{
 		for (int32 SubIndex = 0; SubIndex < Layouts[Index].LayoutBindings.Num(); ++SubIndex)
 		{
-			FDescriptorSetLayoutBinding* Binding = new(OutGfxEntry->DescriptorSetLayoutBindings[Index]) FDescriptorSetLayoutBinding;
-			Binding->ReadFrom(Layouts[Index].LayoutBindings[SubIndex]);
+			FDescriptorSetLayoutBinding& Binding = OutGfxEntry->DescriptorSetLayoutBindings[Index].AddDefaulted_GetRef();
+			Binding.ReadFrom(Layouts[Index].LayoutBindings[SubIndex]);
 		}
 	}
 
