Fix no unity build

diff --git a/Engine/Source/Developer/Android/AndroidTargetPlatformSettings/Private/AndroidTargetPlatformSettings.cpp b/Engine/Source/Developer/Android/AndroidTargetPlatformSettings/Private/AndroidTargetPlatformSettings.cpp
--- a/Engine/Source/Developer/Android/AndroidTargetPlatformSettings/Private/AndroidTargetPlatformSettings.cpp
+++ b/Engine/Source/Developer/Android/AndroidTargetPlatformSettings/Private/AndroidTargetPlatformSettings.cpp
@@ -5,6 +5,7 @@
 #include "UObject/NameTypes.h"
 #include "CoreTypes.h"
 #include "AnalyticsEventAttribute.h"
+#include "HAL/IConsoleManager.h"
 
 FAndroidTargetPlatformSettings::FAndroidTargetPlatformSettings(const TCHAR* CookFlavor, const TCHAR* OverrideIniPlatformName)
 	: TTargetPlatformSettingsBase(CookFlavor, OverrideIniPlatformName)
