Add missing change from CL 30963863
(fixes crash when viewing device output on Android)

diff --git a/Engine/Source/Developer/Android/AndroidTargetPlatformControls/Private/AndroidTargetDevice.h b/Engine/Source/Developer/Android/AndroidTargetPlatformControls/Private/AndroidTargetDevice.h
--- a/Engine/Source/Developer/Android/AndroidTargetPlatformControls/Private/AndroidTargetDevice.h
+++ b/Engine/Source/Developer/Android/AndroidTargetPlatformControls/Private/AndroidTargetDevice.h
@@ -159,6 +159,15 @@ public:
 		return AndroidVersionString;
 	}
 
+	virtual const class ITargetPlatformSettings& GetPlatformSettings() const override
+	{
+		return *(TargetPlatformControls.GetTargetPlatformSettings());
+	}
+	virtual const class ITargetPlatformControls& GetPlatformControls() const override
+	{
+		return TargetPlatformControls;
+	}
+
 	virtual FString GetOperatingSystemName() override;
 
 	virtual int32 GetProcessSnapshot( TArray<FTargetDeviceProcessInfo>& OutProcessInfos ) override;
diff --git a/Engine/Source/Developer/TargetPlatform/Public/Interfaces/ITargetDevice.h b/Engine/Source/Developer/TargetPlatform/Public/Interfaces/ITargetDevice.h
--- a/Engine/Source/Developer/TargetPlatform/Public/Interfaces/ITargetDevice.h
+++ b/Engine/Source/Developer/TargetPlatform/Public/Interfaces/ITargetDevice.h
@@ -338,8 +338,8 @@ protected:
 	};
 public:
 
-	TARGETPLATFORM_API const class ITargetPlatformSettings& GetPlatformSettings() const;
-	TARGETPLATFORM_API const class ITargetPlatformControls& GetPlatformControls() const;
+	TARGETPLATFORM_API virtual const class ITargetPlatformSettings& GetPlatformSettings() const;
+	TARGETPLATFORM_API virtual const class ITargetPlatformControls& GetPlatformControls() const;
 
 	/**
 	 * Checks whether this device is connected.
