Introduce short names for both `AndroidTargetPlatformControls` and `AndroidTargetPlatformSettings`, which are causing certain built targets to run afoul of the `MAX_PATH` limit we impose on UBT file path lengths.

diff --git a/Engine/Source/Developer/Android/AndroidTargetPlatformControls/AndroidTargetPlatformControls.Build.cs b/Engine/Source/Developer/Android/AndroidTargetPlatformControls/AndroidTargetPlatformControls.Build.cs
--- a/Engine/Source/Developer/Android/AndroidTargetPlatformControls/AndroidTargetPlatformControls.Build.cs
+++ b/Engine/Source/Developer/Android/AndroidTargetPlatformControls/AndroidTargetPlatformControls.Build.cs
@@ -8,6 +8,9 @@ public class AndroidTargetPlatformControls : ModuleRules
 	{
 		UnsafeTypeCastWarningLevel = WarningLevel.Error;
 		BinariesSubFolder = "Android";
+		// We need a short name here since this can run afoul very easily of the `MAX_PATH` limit when
+		// combined with building other targets that make use of this.
+		ShortName = "AndTPCon";
 
 		PrivateDependencyModuleNames.AddRange(
 			new string[]
diff --git a/Engine/Source/Developer/Android/AndroidTargetPlatformSettings/AndroidTargetPlatformSettings.Build.cs b/Engine/Source/Developer/Android/AndroidTargetPlatformSettings/AndroidTargetPlatformSettings.Build.cs
--- a/Engine/Source/Developer/Android/AndroidTargetPlatformSettings/AndroidTargetPlatformSettings.Build.cs
+++ b/Engine/Source/Developer/Android/AndroidTargetPlatformSettings/AndroidTargetPlatformSettings.Build.cs
@@ -8,6 +8,9 @@ public class AndroidTargetPlatformSettings : ModuleRules
 	{
 		UnsafeTypeCastWarningLevel = WarningLevel.Error;
 		BinariesSubFolder = "Android";
+		// We need a short name here since this can run afoul very easily of the `MAX_PATH` limit when
+		// combined with building other targets that make use of this.
+		ShortName = "AndTPSet";
 
 		PrivateDependencyModuleNames.AddRange(
 			new string[]
