Fix Android architecture check for project settings

#android

diff --git a/Engine/Source/Developer/Android/AndroidPlatformEditor/Private/AndroidTargetSettingsCustomization.cpp b/Engine/Source/Developer/Android/AndroidPlatformEditor/Private/AndroidTargetSettingsCustomization.cpp
--- a/Engine/Source/Developer/Android/AndroidPlatformEditor/Private/AndroidTargetSettingsCustomization.cpp
+++ b/Engine/Source/Developer/Android/AndroidPlatformEditor/Private/AndroidTargetSettingsCustomization.cpp
@@ -306,8 +306,8 @@ void FAndroidTargetSettingsCustomization::BuildAppManifestSection(IDetailLayoutB
 			.IsEnabled(FEngineBuildSettings::IsSourceDistribution()) \
 			.ToolTip(FEngineBuildSettings::IsSourceDistribution() ? Tip : FAndroidTargetSettingsCustomizationConstants::DisabledTip); \
 	}
-	SETUP_ANDROIDARCH_PROP(TEXT("-arm64"), bBuildForArm64, BuildCategory, LOCTEXT("BuildForArm64ToolTip", "Enable Arm64 CPU architecture support? (use at least NDK r11c, requires Lollipop (android-21) minimum)"));
-	SETUP_ANDROIDARCH_PROP(TEXT("-x64"), bBuildForX8664, BuildCategory, LOCTEXT("BuildForX8664ToolTip", "Enable X86-64 CPU architecture support?"));
+	SETUP_ANDROIDARCH_PROP(TEXT("arm64"), bBuildForArm64, BuildCategory, LOCTEXT("BuildForArm64ToolTip", "Enable Arm64 CPU architecture support? (use at least NDK r11c, requires Lollipop (android-21) minimum)"));
+	SETUP_ANDROIDARCH_PROP(TEXT("x64"), bBuildForX8664, BuildCategory, LOCTEXT("BuildForX8664ToolTip", "Enable X86-64 CPU architecture support?"));
 
 	// @todo android fat binary: Put back in when we expose those
 //	SETUP_SOURCEONLY_PROP(bSplitIntoSeparateApks, BuildCategory, LOCTEXT("SplitIntoSeparateAPKsToolTip", "If checked, CPU architectures and rendering types will be split into separate .apk files"));
diff --git a/Engine/Source/Programs/UnrealBuildTool/Platform/Android/UEBuildAndroid.cs b/Engine/Source/Programs/UnrealBuildTool/Platform/Android/UEBuildAndroid.cs
--- a/Engine/Source/Programs/UnrealBuildTool/Platform/Android/UEBuildAndroid.cs
+++ b/Engine/Source/Programs/UnrealBuildTool/Platform/Android/UEBuildAndroid.cs
@@ -97,7 +97,7 @@ namespace UnrealBuildTool
 	class AndroidArchitectureConfig : UnrealArchitectureConfig
 	{
 		public AndroidArchitectureConfig()
-			: base(UnrealArchitectureMode.SingleTargetLinkSeparately, new[] { UnrealArch.X64, UnrealArch.Arm64 })
+			: base(UnrealArchitectureMode.SingleTargetLinkSeparately, new[] { UnrealArch.Arm64, UnrealArch.X64 })
 		{
 
 		}
@@ -129,15 +129,7 @@ namespace UnrealBuildTool
 				}
 				if (Ini.GetBool("/Script/AndroidRuntimeSettings.AndroidRuntimeSettings", "bBuildForx8664", out bBuild) && bBuild)
 				{
-					if (File.Exists(Path.Combine(Unreal.EngineDirectory.FullName, "Build", "InstalledBuild.txt")))
-					{
-						bUnsupportedBinaryBuildArch = true;
-						Log.TraceWarningOnce("Please install source to build for x86_64 (-x64); ignoring this architecture target.");
-					}
-					else
-					{
-						ActiveArches.Add("x64");
-					}
+					ActiveArches.Add("x64");
 				}
 
 				// we expect one to be specified
