Fixing quick launch not working on Android when using Zen Store and target device and host are on different networks

diff --git a/Engine/Source/Developer/Android/AndroidDeviceDetection/AndroidDeviceDetection.Build.cs b/Engine/Source/Developer/Android/AndroidDeviceDetection/AndroidDeviceDetection.Build.cs
--- a/Engine/Source/Developer/Android/AndroidDeviceDetection/AndroidDeviceDetection.Build.cs
+++ b/Engine/Source/Developer/Android/AndroidDeviceDetection/AndroidDeviceDetection.Build.cs
@@ -13,9 +13,10 @@ public class AndroidDeviceDetection : ModuleRules
 			{
 				"Core",
 				"Json",
-                "JsonUtilities",
-                "PIEPreviewDeviceSpecification"
-            }
+				"JsonUtilities",
+				"PIEPreviewDeviceSpecification",
+				"Zen"
+			}
 		);
 
 		PrivateIncludePathModuleNames.AddRange(
diff --git a/Engine/Source/Developer/Android/AndroidDeviceDetection/Private/AndroidDeviceDetectionModule.cpp b/Engine/Source/Developer/Android/AndroidDeviceDetection/Private/AndroidDeviceDetectionModule.cpp
--- a/Engine/Source/Developer/Android/AndroidDeviceDetection/Private/AndroidDeviceDetectionModule.cpp
+++ b/Engine/Source/Developer/Android/AndroidDeviceDetection/Private/AndroidDeviceDetectionModule.cpp
@@ -25,6 +25,7 @@
 #include "Interfaces/IAndroidDeviceDetection.h"
 #include "Interfaces/IAndroidDeviceDetectionModule.h"
 #include "ITcpMessagingModule.h"
+#include "Experimental/ZenServerInterface.h"
 
 #include "PIEPreviewDeviceSpecification.h"
 #include "JsonObjectConverter.h"
@@ -623,8 +624,11 @@ private:
 				uint16 ReversePortMappings[] = {
 					41899,	// Network file server, DEFAULT_TCP_FILE_SERVING_PORT in NetworkMessage.h
 					1981,	// Unreal Insights data collection, TraceInsightsModule.cpp
-					0		// end of list
-					};
+#if UE_WITH_ZEN
+					UE::Zen::IsDefaultServicePresent() ? UE::Zen::GetDefaultServiceInstance().GetPort() : 0, // Zen Store, usually defaults to 8558
+#endif
+					0 // end of list
+				};
 
 				for (int32 Idx=0; ReversePortMappings[Idx] > 0; Idx++)
 				{
