PR #11615: Fixed Android_Multi builds treating ASTC_NormalRG as ASTC_NormalLA

diff --git a/Engine/Source/Developer/Android/AndroidTargetPlatformSettings/Public/AndroidTargetPlatformSettings.h b/Engine/Source/Developer/Android/AndroidTargetPlatformSettings/Public/AndroidTargetPlatformSettings.h
--- a/Engine/Source/Developer/Android/AndroidTargetPlatformSettings/Public/AndroidTargetPlatformSettings.h
+++ b/Engine/Source/Developer/Android/AndroidTargetPlatformSettings/Public/AndroidTargetPlatformSettings.h
@@ -204,4 +204,15 @@ public:
 	FAndroid_MultiTargetPlatformSettings() : FAndroidTargetPlatformSettings(TEXT("Multi"))
 	{
 	}
+
+	bool SupportsFeature(ETargetPlatformFeatures Feature) const override
+	{
+		//FAndroid_MultiTargetPlatformControls::GetTextureFormats() replaces NormalLA for NormalRG
+		//So we disable the feature support to ensure shaders rebuild the blue channel correctly (See: LA_NORMALMAPS)
+		if (Feature == ETargetPlatformFeatures::NormalmapLAEncodingMode)
+		{
+			return false;
+		}
+		return FAndroidTargetPlatformSettings::SupportsFeature(Feature);
+	}
 };
