AndroidTargetPlatformControls.cpp nop comments only

diff --git a/Engine/Source/Developer/Android/AndroidTargetPlatformControls/Private/AndroidTargetPlatformControls.cpp b/Engine/Source/Developer/Android/AndroidTargetPlatformControls/Private/AndroidTargetPlatformControls.cpp
--- a/Engine/Source/Developer/Android/AndroidTargetPlatformControls/Private/AndroidTargetPlatformControls.cpp
+++ b/Engine/Source/Developer/Android/AndroidTargetPlatformControls/Private/AndroidTargetPlatformControls.cpp
@@ -302,6 +302,8 @@ FName FAndroidTargetPlatformControls::FinalizeVirtualTextureLayerFormat(FName Fo
 
 	// code dupe with IOSTargetPlatform
 
+	// @todo Oodle: restrict this so it's only done when needed for RVT, not for all VT that would be better left as ASTC ; UE-212640
+
 	const static FName VTRemap[][2] =
 	{
 		{ { FName(TEXT("ASTC_RGB")) },			{ AndroidTexFormat::NameETC2_RGB } },
@@ -473,6 +475,8 @@ void FAndroidTargetPlatformControls::GetTextureFormats(const UTexture* Texture,
 
 	// Supported in ES3.2 with ASTC
 	const bool bSupportCompressedVolumeTexture = AndroidTargetPlatformSettings->SupportsTextureFormatCategory(EAndroidTextureFormatCategory::ASTC);
+	// FWIW bSupportCompressedVolumeTexture should be true for Android_DXT but this is setting it to false
+
 	// OpenGL ES has F32 textures but doesn't allow linear filtering unless OES_texture_float_linear
 	const bool bSupportFilteredFloat32Textures = false;
 
