Flavorless 'Android' target platform now defaults to ASTC flavor

diff --git a/Engine/Source/Developer/Android/AndroidTargetPlatformControls/Private/AndroidTargetPlatformControls.h b/Engine/Source/Developer/Android/AndroidTargetPlatformControls/Private/AndroidTargetPlatformControls.h
--- a/Engine/Source/Developer/Android/AndroidTargetPlatformControls/Private/AndroidTargetPlatformControls.h
+++ b/Engine/Source/Developer/Android/AndroidTargetPlatformControls/Private/AndroidTargetPlatformControls.h
@@ -188,7 +188,8 @@ public:
 class FAndroid_ASTCTargetPlatformControls : public FAndroidTargetPlatformControls
 {
 public:
-	FAndroid_ASTCTargetPlatformControls(bool bIsClient, ITargetPlatformSettings* TargetPlatformSettings) : FAndroidTargetPlatformControls(bIsClient, TargetPlatformSettings, TEXT("ASTC"))
+	FAndroid_ASTCTargetPlatformControls(bool bIsClient, ITargetPlatformSettings* TargetPlatformSettings, const TCHAR* OverrideIniPlatformName = TEXT("ASTC")) 
+		: FAndroidTargetPlatformControls(bIsClient, TargetPlatformSettings, OverrideIniPlatformName)
 	{
 	}
 
@@ -210,6 +211,15 @@ public:
 #endif
 };
 
+// Flavorless Android platform defaults to ASTC
+class FAndroid_DefaultTargetPlatformControls : public FAndroid_ASTCTargetPlatformControls
+{
+public:
+	FAndroid_DefaultTargetPlatformControls(bool bIsClient, ITargetPlatformSettings* TargetPlatformSettings) 
+		: FAndroid_ASTCTargetPlatformControls(bIsClient, TargetPlatformSettings, nullptr)
+	{
+	}
+};
 
 class FAndroid_ETC2TargetPlatformControls : public FAndroidTargetPlatformControls
 {
diff --git a/Engine/Source/Developer/Android/AndroidTargetPlatformControls/Private/AndroidTargetPlatformControlsModule.cpp b/Engine/Source/Developer/Android/AndroidTargetPlatformControls/Private/AndroidTargetPlatformControlsModule.cpp
--- a/Engine/Source/Developer/Android/AndroidTargetPlatformControls/Private/AndroidTargetPlatformControlsModule.cpp
+++ b/Engine/Source/Developer/Android/AndroidTargetPlatformControls/Private/AndroidTargetPlatformControlsModule.cpp
@@ -49,9 +49,8 @@ protected:
 				bool bIsClient = Type == 1;
 
 				// flavorless must come first
-				// "flavorless" (aka "vanilla") Android is not well defined and not actually usable
-				//  but Editor wants to see it in the platform list
-				SinglePlatforms.Add(new FAndroidTargetPlatformControls(bIsClient, OutMap[EAndroidTextureFormatCategory::Count], nullptr));
+				// "flavorless" (aka "vanilla") Android defaults to ASTC
+				SinglePlatforms.Add(new FAndroid_DefaultTargetPlatformControls(bIsClient, OutMap[EAndroidTextureFormatCategory::Count]));
 				SinglePlatforms.Add(new FAndroid_ASTCTargetPlatformControls(bIsClient, OutMap[EAndroidTextureFormatCategory::ASTC]));
 				SinglePlatforms.Add(new FAndroid_DXTTargetPlatformControls(bIsClient, OutMap[EAndroidTextureFormatCategory::DXT]));
 				SinglePlatforms.Add(new FAndroid_ETC2TargetPlatformControls(bIsClient, OutMap[EAndroidTextureFormatCategory::ETC2]));
