Fix for landscape weightmap compression not being applied

diff --git a/Engine/Source/Developer/Android/AndroidTargetPlatformControls/Private/AndroidTargetPlatformControls.cpp b/Engine/Source/Developer/Android/AndroidTargetPlatformControls/Private/AndroidTargetPlatformControls.cpp
--- a/Engine/Source/Developer/Android/AndroidTargetPlatformControls/Private/AndroidTargetPlatformControls.cpp
+++ b/Engine/Source/Developer/Android/AndroidTargetPlatformControls/Private/AndroidTargetPlatformControls.cpp
@@ -481,8 +481,8 @@ void FAndroidTargetPlatformControls::GetTextureFormats(const UTexture* Texture,
 	const bool bSupportFilteredFloat32Textures = false;
 
 	// optionaly compress landscape weightmaps for a mobile rendering
-	static const auto CompressLandscapeWeightMapsVar = IConsoleManager::Get().FindTConsoleVariableDataInt(TEXT("r.Mobile.CompressLandscapeWeightMaps"));
-	static const bool bCompressLandscapeWeightMaps = (CompressLandscapeWeightMapsVar && CompressLandscapeWeightMapsVar->GetValueOnAnyThread() != 0);
+	bool bCompressLandscapeWeightMaps = false;
+	GetTargetPlatformSettings()->GetConfigSystem()->GetBool(TEXT("/Script/Engine.RendererSettings"), TEXT("r.Mobile.CompressLandscapeWeightMaps"), bCompressLandscapeWeightMaps, GEngineIni);
 
 	TArray<FName>& LayerFormats = OutFormats.AddDefaulted_GetRef();
 	int32 BlockSize = 1; // this looks wrong? should be 4 for FAndroid_DXTTargetPlatform ? - it is wrong, but BlockSize is ignored
diff --git a/Engine/Source/Developer/IOS/IOSTargetPlatformControls/Private/IOSTargetPlatformControls.cpp b/Engine/Source/Developer/IOS/IOSTargetPlatformControls/Private/IOSTargetPlatformControls.cpp
--- a/Engine/Source/Developer/IOS/IOSTargetPlatformControls/Private/IOSTargetPlatformControls.cpp
+++ b/Engine/Source/Developer/IOS/IOSTargetPlatformControls/Private/IOSTargetPlatformControls.cpp
@@ -507,8 +507,9 @@ void FIOSTargetPlatformControls::GetTextureFormats( const UTexture* Texture, TAr
 	// @todo Oodle: this should not be here
 	//	should be in GetDefaultTextureFormatNamePerLayer
 	//	so that 4x4 checks can be applied correctly, etc.
-	static const auto CompressLandscapeWeightMapsVar = IConsoleManager::Get().FindTConsoleVariableDataInt(TEXT("r.Mobile.CompressLandscapeWeightMaps"));
-	static const bool bCompressLandscapeWeightMaps = (CompressLandscapeWeightMapsVar && CompressLandscapeWeightMapsVar->GetValueOnAnyThread() != 0);
+	bool bCompressLandscapeWeightMaps = false;
+	GetTargetPlatformSettings()->GetConfigSystem()->GetBool(TEXT("/Script/Engine.RendererSettings"), TEXT("r.Mobile.CompressLandscapeWeightMaps"), bCompressLandscapeWeightMaps, GEngineIni);
+
 	if (Texture->LODGroup == TEXTUREGROUP_Terrain_Weightmap && bCompressLandscapeWeightMaps)
 	{
 		TextureFormatNames.Init(FName(TEXT("AutoDXT")), NumLayers);
