Automatically install the Android Emulator during Android setup, and create a device suitable for UE. Add launching an Android emulator image to our Automation command line

#android

diff --git a/Engine/Config/Android/AndroidEmulator.json b/Engine/Config/Android/AndroidEmulator.json
new file mode 100644
--- /dev/null
+++ b/Engine/Config/Android/AndroidEmulator.json
@@ -0,0 +1,6 @@
+{
+	"AvdDevice": "pixel_6",
+	"DeviceApiVersion": "35",
+	"DeviceDataSize": "48G",
+	"DeviceRamSize": "6144"
+}
diff --git a/Engine/Extras/Android/SetupAndroidEmulator.bat b/Engine/Extras/Android/SetupAndroidEmulator.bat
new file mode 100755
--- /dev/null
+++ b/Engine/Extras/Android/SetupAndroidEmulator.bat
@@ -0,0 +1,122 @@
+@echo off
+setlocal ENABLEEXTENSIONS ENABLEDELAYEDEXPANSION
+
+set AVD_DEVICE=%1
+set DEVICE_API_VERSION=%2
+set DEVICE_DATA_SIZE=%3
+set DEVICE_RAM_SIZE=%4
+
+if "%5" equ "-noninteractive" (
+	set PAUSE=
+) else (
+	set PAUSE=pause
+)
+
+rem hardcoded versions for compatibility with non-Turnkey manual running
+if "%AVD_DEVICE%" equ "" set AVD_DEVICE=pixel_6
+if "%DEVICE_API_VERSION%" equ "" set DEVICE_API_VERSION=35
+if "%DEVICE_DATA_SIZE%" equ "" set DEVICE_DATA_SIZE=48G
+if "%DEVICE_RAM_SIZE%" equ "" set DEVICE_RAM_SIZE=6144
+
+if not exist "%ANDROID_HOME%" (
+	echo Android SDK not found at: %ANDROID_HOME%
+	echo Unable to locate local Android SDK location. Did you run Android Studio after installing?
+	echo If Android Studio is installed, please run again with SDK path as parameter, otherwise download Android Studio 2022.2.1 from https://developer.android.com/studio/archive
+	%PAUSE%
+	exit /b 1
+)
+
+echo Android Studio SDK Path: %ANDROID_HOME%
+
+set CMDLINETOOLS_PATH=%ANDROID_HOME%\cmdline-tools\latest\bin
+if not exist "%CMDLINETOOLS_PATH%" (
+	set CMDLINETOOLS_PATH=%ANDROID_HOME%\tools\bin
+	if not exist "!CMDLINETOOLS_PATH!" (
+		echo Unable to locate sdkmanager.bat. Did you run Android Studio and install cmdline-tools after installing?
+		%PAUSE%
+		exit /b 2
+	)
+)
+
+if /i "%PROCESSOR_ARCHITECTURE%" equ "AMD64" (
+	set ABI=x86_64
+) else if /i "%PROCESSOR_ARCHITECTURE%" equ "ARM64" (
+	set ABI=arm64-v8a
+) else (
+	echo Unsupported architecture %ABI%.
+	%PAUSE%
+	exit /b 3
+)
+
+set AVD_TAG=google_apis_playstore
+set AVD_PACKAGE=system-images^;android-%DEVICE_API_VERSION%^;%AVD_TAG%^;%ABI%
+
+call "%CMDLINETOOLS_PATH%\sdkmanager.bat" "emulator" "%AVD_PACKAGE%" || (
+	echo Update failed. Please check the Android Studio install.
+	%PAUSE%
+	exit /b 4
+)
+
+set EMULATOR_PATH=%ANDROID_HOME%\emulator
+if not exist "%EMULATOR_PATH%" (
+	echo Update failed. Did you accept the license agreement?
+	%PAUSE%
+	exit /b 5
+)
+
+set DEVICE_NAME=UE_%AVD_DEVICE%_API_%DEVICE_API_VERSION%
+
+for /f "tokens=*" %%G in ('%EMULATOR_PATH%\emulator.exe -list-avds') do (
+	if "%%G" equ "%DEVICE_NAME%" (
+		echo Android virtual device %DEVICE_NAME% already exists...creation skipped.
+		goto :DONE
+	)
+)
+
+call "%CMDLINETOOLS_PATH%\avdmanager.bat" create avd -n "%DEVICE_NAME%" -k "%AVD_PACKAGE%" -g "%AVD_TAG%" -b "%ABI%" -d "%AVD_DEVICE%" -f || (
+	echo Android virtual device %DEVICE_NAME% creation failed.
+	%PAUSE%
+	exit /b 6
+)
+
+echo Android virtual device %DEVICE_NAME% created.
+
+if "%ANDROID_USER_HOME%" equ "" set ANDROID_USER_HOME=%USERPROFILE%\.android
+if "%ANDROID_EMULATOR_HOME%" equ "" set ANDROID_EMULATOR_HOME=%ANDROID_USER_HOME%
+if "%ANDROID_AVD_HOME%" equ "" set ANDROID_AVD_HOME=%ANDROID_EMULATOR_HOME%\avd
+
+set DEVICE_PATH=%ANDROID_AVD_HOME%\%DEVICE_NAME%.avd
+set DEVICE_CONFIG=config.ini
+set DEVICE_CONFIG_PATH=%DEVICE_PATH%\%DEVICE_CONFIG%
+set DEVICE_TEMP_CONFIG_PATH=%DEVICE_CONFIG_PATH%.tmp
+
+(echo AvdId=%DEVICE_NAME% & echo avd.ini.displayname=%DEVICE_NAME:_= %) > "%DEVICE_TEMP_CONFIG_PATH%"
+
+for /f "tokens=1* delims=\=" %%G in (%DEVICE_CONFIG_PATH%) do (
+	(if "%%G" equ "PlayStore.enabled" (
+		echo %%G=yes
+	) else if "%%G" equ "disk.dataPartition.size" (
+		echo %%G=%DEVICE_DATA_SIZE%
+	) else if "%%G" equ "hw.gpu.enabled" (
+		echo %%G=yes
+	) else if "%%G" equ "hw.initialOrientation" (
+		echo %%G=landscape
+	) else if "%%G" equ "hw.keyboard" (
+		echo %%G=yes
+	) else if "%%G" equ "hw.ramSize" (
+		echo %%G=%DEVICE_RAM_SIZE%
+	) else (
+		echo %%G=%%H
+	)) >> "%DEVICE_TEMP_CONFIG_PATH%"
+)
+
+move /y "%DEVICE_TEMP_CONFIG_PATH%" "%DEVICE_CONFIG_PATH%" || (
+	echo Android virtual device %DEVICE_NAME% %DEVICE_CONFIG% update failed.
+	%PAUSE%
+	exit /b 7
+)
+
+:DONE
+echo Success^^!
+%PAUSE%
+exit /b 0
diff --git a/Engine/Extras/Android/SetupAndroidEmulator.command b/Engine/Extras/Android/SetupAndroidEmulator.command
new file mode 100644
--- /dev/null
+++ b/Engine/Extras/Android/SetupAndroidEmulator.command
@@ -0,0 +1,146 @@
+#!/usr/bin/env bash
+
+pause()
+{
+	read -rsp "Press any key to continue . . . " -n1
+	echo
+}
+
+AVD_DEVICE=${1:-}
+DEVICE_API_VERSION=${2:-}
+DEVICE_DATA_SIZE=${3:-}
+DEVICE_RAM_SIZE=${4:-}
+
+if [ "${5:-}" == "-noninteractive" ]; then
+    PAUSE=
+else
+    PAUSE="pause"
+fi
+
+# hardcoded versions for compatibility with non-Turnkey manual running
+if [ -z "${AVD_DEVICE}" ]; then
+    AVD_DEVICE="pixel_6"
+fi
+if [ -z "${DEVICE_API_VERSION}" ]; then
+    DEVICE_API_VERSION="35"
+fi
+if [ -z "${DEVICE_DATA_SIZE}" ]; then
+    DEVICE_DATA_SIZE="48G"
+fi
+if [ -z "${DEVICE_RAM_SIZE}" ]; then
+    DEVICE_RAM_SIZE="6144"
+fi
+
+if [ ! -d "${ANDROID_HOME}" ]; then
+	echo "Android SDK not found at: ${ANDROID_HOME}"
+	echo "Unable to locate local Android SDK location. Did you run Android Studio after installing?"
+	echo "If Android Studio is installed, please run again with SDK path as parameter, otherwise download Android Studio 2022.2.1 from https://developer.android.com/studio/archive"
+	${PAUSE}
+	exit 1
+fi
+
+echo "Android Studio SDK Path: ${ANDROID_HOME}"
+
+CMDLINETOOLS_PATH="${ANDROID_HOME}/cmdline-tools/latest/bin"
+if [ ! -d "$CMDLINETOOLS_PATH" ]; then
+	CMDLINETOOLS_PATH="${ANDROID_HOME}/tools/bin"
+	if [ ! -d "${CMDLINETOOLS_PATH}" ]; then
+		echo "Unable to locate sdkmanager. Did you run Android Studio and install cmdline-tools after installing?"
+		${PAUSE}
+		exit 1
+	fi
+fi
+
+ABI=$(uname -m)
+if [ "${ABI}" != "x86_64" ] && [ "${ABI}" != "arm64-v8a" ]; then
+	echo "Unsupported architecture ${ABI}."
+	${PAUSE}
+	exit 1
+fi
+
+AVD_TAG="google_apis_playstore"
+AVD_PACKAGE="system-images;android-${DEVICE_API_VERSION};${AVD_TAG};${ABI}"
+
+if ! "${CMDLINETOOLS_PATH}/sdkmanager" "emulator" "${AVD_PACKAGE}"; then
+	echo "Update failed. Please check the Android Studio install."
+	${PAUSE}
+	exit 1
+fi
+
+EMULATOR_PATH="${ANDROID_HOME}/emulator"
+if [ ! -d "${EMULATOR_PATH}" ]; then
+	echo "Update failed. Did you accept the license agreement?"
+	${PAUSE}
+	exit 1
+fi
+
+DEVICE_NAME="UE_${AVD_DEVICE}_API_${DEVICE_API_VERSION}"
+for LINE in $("${EMULATOR_PATH}/emulator" -list-avds); do
+	if [ "${LINE}" == "${DEVICE_NAME}" ]; then
+		echo "Android virtual device ${DEVICE_NAME} already exists...creation skipped."
+		DEVICE_NAME=""
+		break
+	fi
+done
+
+if [ -n "${DEVICE_NAME}" ]; then
+	if ! "${CMDLINETOOLS_PATH}/avdmanager" create avd -n "${DEVICE_NAME}" -k "${AVD_PACKAGE}" -g "${AVD_TAG}" -b "${ABI}" -d "${AVD_DEVICE}" -f; then
+		echo "Android virtual device ${DEVICE_NAME} creation failed."
+		${PAUSE}
+		exit 1
+	fi
+	
+	echo "Android virtual device ${DEVICE_NAME} created."
+	
+	if [ -z "${ANDROID_USER_HOME}" ]; then
+		ANDROID_USER_HOME="${HOME}/.android"
+	fi
+	if [ -z "${ANDROID_EMULATOR_HOME}" ]; then
+		ANDROID_EMULATOR_HOME="${ANDROID_USER_HOME}"
+	fi
+	if [ -z "${ANDROID_AVD_HOME}" ]; then
+		ANDROID_AVD_HOME="${ANDROID_EMULATOR_HOME}/avd"
+	fi
+
+	DEVICE_PATH="${ANDROID_AVD_HOME}/${DEVICE_NAME}.avd"
+	DEVICE_CONFIG="config.ini"
+	DEVICE_CONFIG_PATH="${DEVICE_PATH}/${DEVICE_CONFIG}"
+	DEVICE_TEMP_CONFIG_PATH="${DEVICE_PATH}/${DEVICE_CONFIG}.tmp"
+
+	{
+		echo "AvdId=${DEVICE_NAME}"
+		echo "avd.ini.displayname=${DEVICE_NAME//_/ }"
+	} > "${DEVICE_TEMP_CONFIG_PATH}"
+	
+	while read -r LINE; do
+		if [[ "${LINE}" =~ ^([^=]+)=([^=]+)$ ]]; then
+			if [ "${BASH_REMATCH[1]}" == "PlayStore.enabled" ]; then
+				echo "${BASH_REMATCH[1]}=yes"
+			elif [ "${BASH_REMATCH[1]}" == "disk.dataPartition.size" ]; then
+				echo "${BASH_REMATCH[1]}=${DEVICE_DATA_SIZE}"
+			elif [ "${BASH_REMATCH[1]}" == "hw.gpu.enabled" ]; then
+				echo "${BASH_REMATCH[1]}=yes"
+			elif [ "${BASH_REMATCH[1]}" == "hw.initialOrientation" ]; then
+				echo "${BASH_REMATCH[1]}=landscape"
+			elif [ "${BASH_REMATCH[1]}" == "hw.keyboard" ]; then
+				echo "${BASH_REMATCH[1]}=yes"
+			elif [ "${BASH_REMATCH[1]}" == "hw.ramSize" ]; then
+				echo "${BASH_REMATCH[1]}=${DEVICE_RAM_SIZE}"
+			else
+				echo "${BASH_REMATCH[1]}=${BASH_REMATCH[2]}"
+			fi
+		else
+			echo "${LINE}="
+		fi >> "${DEVICE_TEMP_CONFIG_PATH}"
+	done < "${DEVICE_CONFIG_PATH}"
+
+	if ! mv -f "${DEVICE_TEMP_CONFIG_PATH}" "${DEVICE_CONFIG_PATH}"; then
+		echo "Android virtual device ${DEVICE_NAME} ${DEVICE_CONFIG} update failed."
+		${PAUSE}
+		exit 1
+	fi
+fi
+
+echo Success!
+${PAUSE}
+exit 0
diff --git a/Engine/Extras/Android/SetupAndroidEmulator.sh b/Engine/Extras/Android/SetupAndroidEmulator.sh
new file mode 100644
--- /dev/null
+++ b/Engine/Extras/Android/SetupAndroidEmulator.sh
@@ -0,0 +1,151 @@
+#!/bin/bash
+
+if [[ "${OSTYPE}" == "darwin"* ]]; then
+	echo "Please run SetupAndroidEmulator.command on MacOSX; attempting to run it for you."
+	DIR="$( cd "$( dirname "${BASH_SOURCE[0]}" )" >/dev/null 2>&1 && pwd )"
+	exec "${DIR}/SetupAndroidEmulator.command" "${@}"
+	exit 1
+fi
+
+pause()
+{
+	read -rsp "Press any key to continue . . . " -n1
+	echo
+}
+
+AVD_DEVICE=${1:-}
+DEVICE_API_VERSION=${2:-}
+DEVICE_DATA_SIZE=${3:-}
+DEVICE_RAM_SIZE=${4:-}
+
+if [ "${9:-}" == "-noninteractive" ]; then
+    PAUSE=
+else
+    PAUSE="pause"
+fi
+
+# hardcoded versions for compatibility with non-Turnkey manual running
+if [ -z "${AVD_DEVICE}" ]; then
+    AVD_DEVICE="pixel_6"
+fi
+if [ -z "${DEVICE_API_VERSION}" ]; then
+    DEVICE_API_VERSION="35"
+fi
+if [ -z "${DEVICE_DATA_SIZE}" ]; then
+    DEVICE_DATA_SIZE="48G"
+fi
+if [ -z "${DEVICE_RAM_SIZE}" ]; then
+    DEVICE_RAM_SIZE="6144"
+fi
+
+if [ ! -d "${ANDROID_HOME}" ]; then
+	echo "Android SDK not found at: ${ANDROID_HOME}"
+	echo "Unable to locate local Android SDK location. Did you run Android Studio after installing?"
+	echo "If Android Studio is installed, please run again with SDK path as parameter, otherwise download Android Studio 2022.2.1 from https://developer.android.com/studio/archive"
+	${PAUSE}
+	exit 1
+fi
+echo "Android Studio SDK Path: ${ANDROID_HOME}"
+
+CMDLINETOOLS_PATH="${ANDROID_HOME}/cmdline-tools/latest/bin"
+if [ ! -d "${CMDLINETOOLS_PATH}" ]; then
+	CMDLINETOOLS_PATH="${ANDROID_HOME}/tools/bin"
+	if [ ! -d "${CMDLINETOOLS_PATH}" ]; then
+		echo "Unable to locate sdkmanager. Did you run Android Studio and install cmdline-tools after installing?"
+		${PAUSE}
+		exit 1
+	fi
+fi
+
+ABI=$(uname -m)
+if [ "${ABI}" != "x86_64" ] && [ "${ABI}" != "arm64-v8a" ]; then
+	echo "Unsupported architecture ${ABI}."
+	${PAUSE}
+	exit 1
+fi
+
+AVD_TAG="google_apis_playstore"
+AVD_PACKAGE="system-images;android-${DEVICE_API_VERSION};${AVD_TAG};${ABI}"
+
+if ! "${CMDLINETOOLS_PATH}/sdkmanager" "emulator" "${AVD_PACKAGE}"; then
+	echo "Update failed. Please check the Android Studio install."
+	${PAUSE}
+	exit 1
+fi
+
+EMULATOR_PATH="${ANDROID_HOME}/emulator"
+if [ ! -d "${EMULATOR_PATH}" ]; then
+	echo "Update failed. Did you accept the license agreement?"
+	${PAUSE}
+	exit 1
+fi
+
+DEVICE_NAME="UE_${AVD_DEVICE}_API_${DEVICE_API_VERSION}"
+for LINE in $("${EMULATOR_PATH}/emulator" -list-avds); do
+	if [ "${LINE}" == "${DEVICE_NAME}" ]; then
+		echo "Android virtual device ${DEVICE_NAME} already exists...creation skipped."
+		DEVICE_NAME=""
+	fi
+done
+
+if [ -n "${DEVICE_NAME}" ]; then
+	if ! "${CMDLINETOOLS_PATH}/avdmanager" create avd -n "${DEVICE_NAME}" -k "${AVD_PACKAGE}" -g "${AVD_TAG}" -b "${ABI}" -d "${AVD_DEVICE}" -f; then
+		echo "Android virtual device ${DEVICE_NAME} creation failed."
+		${PAUSE}
+		exit 1
+	fi
+	
+	echo "Android virtual device ${DEVICE_NAME} created."
+	
+	if [ -z "${ANDROID_USER_HOME}" ]; then
+		ANDROID_USER_HOME="${HOME}/.android"
+	fi
+	if [ -z "${ANDROID_EMULATOR_HOME}" ]; then
+		ANDROID_EMULATOR_HOME="${ANDROID_USER_HOME}"
+	fi
+	if [ -z "${ANDROID_AVD_HOME}" ]; then
+		ANDROID_AVD_HOME="${ANDROID_EMULATOR_HOME}/avd"
+	fi
+
+	DEVICE_PATH="${ANDROID_AVD_HOME}/${DEVICE_NAME}.avd"
+	DEVICE_CONFIG="config.ini"
+	DEVICE_CONFIG_PATH="${DEVICE_PATH}/${DEVICE_CONFIG}"
+	DEVICE_TEMP_CONFIG_PATH="${DEVICE_PATH}/${DEVICE_CONFIG}.tmp"
+
+	{
+		echo "AvdId=${DEVICE_NAME}"
+		echo "avd.ini.displayname=${DEVICE_NAME//_/ }"
+	} > "${DEVICE_TEMP_CONFIG_PATH}"
+	
+	while read -r LINE; do
+		if [[ "${LINE}" =~ ^([^=]+)=([^=]+)$ ]]; then
+			if [ "${BASH_REMATCH[1]}" == "PlayStore.enabled" ]; then
+				echo "${BASH_REMATCH[1]}=yes"
+			elif [ "${BASH_REMATCH[1]}" == "disk.dataPartition.size" ]; then
+				echo "${BASH_REMATCH[1]}=${DEVICE_DATA_SIZE}"
+			elif [ "${BASH_REMATCH[1]}" == "hw.gpu.enabled" ]; then
+				echo "${BASH_REMATCH[1]}=yes"
+			elif [ "${BASH_REMATCH[1]}" == "hw.initialOrientation" ]; then
+				echo "${BASH_REMATCH[1]}=landscape"
+			elif [ "${BASH_REMATCH[1]}" == "hw.keyboard" ]; then
+				echo "${BASH_REMATCH[1]}=yes"
+			elif [ "${BASH_REMATCH[1]}" == "hw.ramSize" ]; then
+				echo "${BASH_REMATCH[1]}=${DEVICE_RAM_SIZE}"
+			else
+				echo "${BASH_REMATCH[1]}=${BASH_REMATCH[2]}"
+			fi
+		else
+			echo "${LINE}="
+		fi >> "${DEVICE_TEMP_CONFIG_PATH}"
+	done < "${DEVICE_CONFIG_PATH}"
+
+	if ! mv -f "${DEVICE_TEMP_CONFIG_PATH}" "${DEVICE_CONFIG_PATH}"; then
+		echo "Android virtual device ${DEVICE_NAME} ${DEVICE_CONFIG} update failed."
+		${PAUSE}
+		exit 1
+	fi
+fi
+
+echo Success!
+${PAUSE}
+exit 0
diff --git a/Engine/Source/Developer/Android/AndroidDeviceDetection/Private/AndroidDeviceDetectionModule.cpp b/Engine/Source/Developer/Android/AndroidDeviceDetection/Private/AndroidDeviceDetectionModule.cpp
--- a/Engine/Source/Developer/Android/AndroidDeviceDetection/Private/AndroidDeviceDetectionModule.cpp
+++ b/Engine/Source/Developer/Android/AndroidDeviceDetection/Private/AndroidDeviceDetectionModule.cpp
@@ -112,10 +112,11 @@ public:
 		return 0;
 	}
 
-	void UpdateADBPath(FString &InADBPath, FString& InGetPropCommand, bool InbGetExtensionsViaSurfaceFlinger)
+	void UpdatePaths(FString InADBPath, FString InAvdHomePath, FString InGetPropCommand, bool InbGetExtensionsViaSurfaceFlinger)
 	{
-		ADBPath = InADBPath;
-		GetPropCommand = InGetPropCommand;
+		ADBPath = MoveTemp(InADBPath);
+		AvdHomePath = MoveTemp(InAvdHomePath);
+		GetPropCommand = MoveTemp(InGetPropCommand);
 		bGetExtensionsViaSurfaceFlinger = InbGetExtensionsViaSurfaceFlinger;
 
 		HasADBPath = !ADBPath.IsEmpty();
@@ -265,61 +266,104 @@ private:
 		StdOut = StdOut.Replace(TEXT("\r"), TEXT("\n"));
 		StdOut.ParseIntoArray(DeviceStrings, TEXT("\n"), true);
 
+		TArray<FString> AvdNames;
+		IFileManager::Get().FindFiles(AvdNames, *AvdHomePath, TEXT("ini"));
+
+		for (FString& AvdName : AvdNames)
+		{
+			AvdName.LeftChopInline(4);
+		}
+
 		// list of any existing port forwardings, filled in when we find a device we need to add.
 		TArray<FString> PortForwardings;
 
 		// a list containing all devices found this time, so we can remove anything not in this list
-		TArray<FString> CurrentlyConnectedDevices;
+		TSet<FStringView> CurrentlyConnectedDevices;
 
-		for (int32 StringIndex = 0; StringIndex < DeviceStrings.Num(); ++StringIndex)
+		for (const FString& DeviceString : DeviceStrings)
 		{
-			const FString& DeviceString = DeviceStrings[StringIndex];
-
 			// skip over non-device lines
 			if (DeviceString.StartsWith("* ") || DeviceString.StartsWith("List "))
 			{
 				continue;
 			}
 
-			// grab the device serial number
 			int32 TabIndex;
 
-			// use either tab or space as separator
-			if (!DeviceString.FindChar(TCHAR('\t'), TabIndex))
+			if (!DeviceString.FindChar(TCHAR(' '), TabIndex) && !DeviceString.FindChar(TCHAR('\t'), TabIndex))
 			{
-				if (!DeviceString.FindChar(TCHAR(' '), TabIndex))
-				{
-					continue;
-				}
+				continue;
 			}
 
-			FAndroidDeviceInfo NewDeviceInfo;
+			FStringView Status = FStringView{DeviceString}.Mid(TabIndex + 1).TrimStart();
 
-			NewDeviceInfo.SerialNumber = DeviceString.Left(TabIndex);
-			const FString DeviceState = DeviceString.Mid(TabIndex + 1).TrimStart();
-
-			NewDeviceInfo.bAuthorizedDevice = DeviceState != TEXT("unauthorized");
+			if (int32 TabIndex2; Status.FindChar(TCHAR(' '), TabIndex2) || Status.FindChar(TCHAR('\t'), TabIndex2))
+			{
+				Status.LeftInline(TabIndex2);
+			}
 
-			// add it to our list of currently connected devices
-			CurrentlyConnectedDevices.Add(NewDeviceInfo.SerialNumber);
+			bool bUnauthorized = Status == TEXT("unauthorized");
 
-			// move on to next device if this one is already a known device that has either already been authorized or the authorization
-			// status has not changed
-			if (DeviceMap.Contains(NewDeviceInfo.SerialNumber) && 
-				(DeviceMap[NewDeviceInfo.SerialNumber].bAuthorizedDevice == NewDeviceInfo.bAuthorizedDevice))
+			if (!bUnauthorized && Status != TEXT("device"))
 			{
 				continue;
 			}
 
-			if (!NewDeviceInfo.bAuthorizedDevice)
+			FString SerialNumber = DeviceString.Left(TabIndex), DeviceId, AvdName;
+			FAndroidDeviceInfo* Device;
+			FAndroidDeviceInfo NewDeviceInfo;
+
+			// Find the AVD name
+			if (FString AvdNameOutput; SerialNumber.StartsWith("emulator-") && ExecuteAdbCommand(*(TEXT("-s ") + SerialNumber + TEXT(" emu avd name")), &AvdNameOutput, nullptr))
+			{
+				const TCHAR* Ptr = *AvdNameOutput;
+				FParse::Line(&Ptr, AvdName);
+				DeviceId = TEXT("avd-") + AvdName;
+				
+				AvdNames.RemoveSingle(AvdName);
+			}
+			else
+			{
+				DeviceId = SerialNumber;
+			}
+			
+			if (bUnauthorized)
 			{
-				//note: AndroidTargetDevice::GetName() does not fetch this value, do not rely on this
-				NewDeviceInfo.DeviceName = TEXT("Unauthorized - enable USB debugging");
+				Device = DeviceMap.Find(DeviceId);
+			
+				if (!Device || Device->bAuthorizedDevice)
+				{
+					NewDeviceInfo.DeviceId = MoveTemp(DeviceId);
+					NewDeviceInfo.AvdName = MoveTemp(AvdName);
+					NewDeviceInfo.bAuthorizedDevice = false;
+					//note: AndroidTargetDevice::GetName() does not fetch this value, do not rely on this
+					NewDeviceInfo.DeviceName = TEXT("Unauthorized - enable USB debugging");
+					
+					NewDeviceInfo.SerialNumber = MoveTemp(SerialNumber);
+			
+					// add the device to the map
+					{
+						FScopeLock ScopeLock(DeviceMapLock);
+
+						if (!Device)
+						{
+							Device = &DeviceMap.Add(NewDeviceInfo.DeviceId);
+						}
+
+						*Device = MoveTemp(NewDeviceInfo);
+					}
+				}
 			}
 			else
 			{
+				Device = nullptr;
+				
+				NewDeviceInfo.DeviceId = MoveTemp(DeviceId);
+				NewDeviceInfo.AvdName = MoveTemp(AvdName);
+				NewDeviceInfo.bAuthorizedDevice = true;
+
 				// grab the Android version
-				const FString AndroidVersionCommand = FString::Printf(TEXT("-s %s %s ro.build.version.release"), *NewDeviceInfo.SerialNumber, *GetPropCommand);
+				const FString AndroidVersionCommand = FString::Printf(TEXT("-s %s %s ro.build.version.release"), *SerialNumber, *GetPropCommand);
 				if (!ExecuteAdbCommand(*AndroidVersionCommand, &NewDeviceInfo.HumanAndroidVersion, nullptr))
 				{
 					continue;
@@ -328,7 +372,7 @@ private:
 				NewDeviceInfo.HumanAndroidVersion.TrimStartAndEndInline();
 
 				// grab the Android SDK version
-				const FString SDKVersionCommand = FString::Printf(TEXT("-s %s %s ro.build.version.sdk"), *NewDeviceInfo.SerialNumber, *GetPropCommand);
+				const FString SDKVersionCommand = FString::Printf(TEXT("-s %s %s ro.build.version.sdk"), *SerialNumber, *GetPropCommand);
 				FString SDKVersionString;
 				if (!ExecuteAdbCommand(*SDKVersionCommand, &SDKVersionString, nullptr))
 				{
@@ -343,7 +387,7 @@ private:
 				if (bGetExtensionsViaSurfaceFlinger)
 				{
 					// get the GL extensions string (and a bunch of other stuff)
-					const FString ExtensionsCommand = FString::Printf(TEXT("-s %s shell dumpsys SurfaceFlinger"), *NewDeviceInfo.SerialNumber);
+					const FString ExtensionsCommand = FString::Printf(TEXT("-s %s shell dumpsys SurfaceFlinger"), *SerialNumber);
 					if (!ExecuteAdbCommand(*ExtensionsCommand, &NewDeviceInfo.GLESExtensions, nullptr))
 					{
 						continue;
@@ -364,7 +408,7 @@ private:
 
 				// grab device brand
 				{
-					FString ExecCommand = FString::Printf(TEXT("-s %s %s ro.product.brand"), *NewDeviceInfo.SerialNumber, *GetPropCommand);
+					FString ExecCommand = FString::Printf(TEXT("-s %s %s ro.product.brand"), *SerialNumber, *GetPropCommand);
 
 					FString RoProductBrand;
 					ExecuteAdbCommand(*ExecCommand, &RoProductBrand, nullptr);
@@ -375,7 +419,7 @@ private:
 				// grab screen resolution
 				{
 					FString ResolutionString;
-					const FString ExecCommand = FString::Printf(TEXT("-s %s shell wm size"), *NewDeviceInfo.SerialNumber);
+					const FString ExecCommand = FString::Printf(TEXT("-s %s shell wm size"), *SerialNumber);
 					if (ExecuteAdbCommand(*ExecCommand, &ResolutionString, nullptr))
 					{
 						bool bFoundResX = false;
@@ -408,7 +452,7 @@ private:
 
 				// grab the GL ES version
 				FString GLESVersionString;
-				const FString GLVersionCommand = FString::Printf(TEXT("-s %s %s ro.opengles.version"), *NewDeviceInfo.SerialNumber, *GetPropCommand);
+				const FString GLVersionCommand = FString::Printf(TEXT("-s %s %s ro.opengles.version"), *SerialNumber, *GetPropCommand);
 				if (!ExecuteAdbCommand(*GLVersionCommand, &GLESVersionString, nullptr))
 				{
 					continue;
@@ -418,9 +462,9 @@ private:
 				// Find the device model
 				FParse::Value(*DeviceString, TEXT("model:"), NewDeviceInfo.Model);
 				// find the product model (this must match java's android.os.build.model)
-				FString ModelCommand = FString::Printf(TEXT("-s %s %s ro.product.model"), *NewDeviceInfo.SerialNumber, *GetPropCommand);
+				FString ModelCommand = FString::Printf(TEXT("-s %s %s ro.product.model"), *SerialNumber, *GetPropCommand);
 				FString RoProductModel;
-				if( ExecuteAdbCommand(*ModelCommand, &RoProductModel, nullptr) )
+				if (ExecuteAdbCommand(*ModelCommand, &RoProductModel, nullptr) )
 				{
 					if(!RoProductModel.IsEmpty())
 					{
@@ -430,7 +474,7 @@ private:
 				
 				// Find the build ID
 				FString BuildNumberString;
-				const FString BuildNumberCommand = FString::Printf(TEXT("-s %s %s ro.build.display.id"), *NewDeviceInfo.SerialNumber, *GetPropCommand);
+				const FString BuildNumberCommand = FString::Printf(TEXT("-s %s %s ro.build.display.id"), *SerialNumber, *GetPropCommand);
 				if (ExecuteAdbCommand(*BuildNumberCommand, &BuildNumberString, nullptr))
 				{
 					NewDeviceInfo.BuildNumber = BuildNumberString.TrimStartAndEnd();
@@ -457,7 +501,7 @@ private:
 				auto MinorVK = [](uint32 Version) { return (((uint32_t)(Version) >> 12) & 0x3FFU); };
 				auto PatchVK = [](uint32 Version) { return ((uint32_t)(Version) & 0xFFFU); };
 				FString FeaturesString;
-				const FString FeaturesStringCommand = FString::Printf(TEXT("-s %s shell pm list features"), *NewDeviceInfo.SerialNumber);
+				const FString FeaturesStringCommand = FString::Printf(TEXT("-s %s shell pm list features"), *SerialNumber);
 				if (ExecuteAdbCommand(*FeaturesStringCommand, &FeaturesString, nullptr))
 				{
 					FString VulkanVersionLine = FindLineContaining(FeaturesString, TEXT("android.hardware.vulkan.version"));
@@ -472,7 +516,7 @@ private:
 
 				// try vkjson:
 				FString VKJsonString;
-				const FString VKJsonStringCommand = FString::Printf(TEXT("-s %s shell cmd gpu vkjson"), *NewDeviceInfo.SerialNumber);
+				const FString VKJsonStringCommand = FString::Printf(TEXT("-s %s shell cmd gpu vkjson"), *SerialNumber);
 				if (ExecuteAdbCommand(*VKJsonStringCommand, &VKJsonString, nullptr))
 				{
 					FString VulkanVersionLine = FindLineContaining(VKJsonString, TEXT("apiVersion"));
@@ -490,13 +534,13 @@ private:
 					}
 				}
 				
-				if(NewDeviceInfo.VulkanVersion.IsEmpty())
+				if (NewDeviceInfo.VulkanVersion.IsEmpty())
 				{
 					NewDeviceInfo.VulkanVersion = TEXT("0.0.0");
 				}
 
 				// create the hardware field
-				FString HardwareCommand = FString::Printf(TEXT("-s %s %s ro.hardware"), *NewDeviceInfo.SerialNumber, *GetPropCommand);
+				FString HardwareCommand = FString::Printf(TEXT("-s %s %s ro.hardware"), *SerialNumber, *GetPropCommand);
 				FString RoHardware;
 				{
 					ExecuteAdbCommand(*HardwareCommand, &RoHardware, nullptr);
@@ -505,13 +549,13 @@ private:
 				}
 				if (RoHardware.Contains(TEXT("qcom")))
 				{
-					HardwareCommand = FString::Printf(TEXT("-s %s %s ro.hardware.chipname"), *NewDeviceInfo.SerialNumber, *GetPropCommand);
+					HardwareCommand = FString::Printf(TEXT("-s %s %s ro.hardware.chipname"), *SerialNumber, *GetPropCommand);
 					ExecuteAdbCommand(*HardwareCommand, &RoHardware, nullptr);
 					const TCHAR* Ptr = *RoHardware;
 					FParse::Line(&Ptr, NewDeviceInfo.Hardware);
 				}
 				{
-					HardwareCommand = FString::Printf(TEXT("-s %s %s ro.soc.model"), *NewDeviceInfo.SerialNumber, *GetPropCommand);
+					HardwareCommand = FString::Printf(TEXT("-s %s %s ro.soc.model"), *SerialNumber, *GetPropCommand);
 					FString RoSOCModelIn;
 					FString RoSOCModelOut;
 					ExecuteAdbCommand(*HardwareCommand, &RoSOCModelIn, nullptr);
@@ -525,7 +569,7 @@ private:
 
 				// Read hardware from cpuinfo:
 				FString CPUInfoString;
-				const FString CPUInfoCommand = FString::Printf(TEXT("-s %s shell cat /proc/cpuinfo"), *NewDeviceInfo.SerialNumber);
+				const FString CPUInfoCommand = FString::Printf(TEXT("-s %s shell cat /proc/cpuinfo"), *SerialNumber);
 				if (ExecuteAdbCommand(*CPUInfoCommand, &CPUInfoString, nullptr))
 				{
 					FString HardwareLine = FindLineContaining(CPUInfoString, TEXT("Hardware"));
@@ -540,7 +584,7 @@ private:
 
 				// Total physical mem:
 				FString MemTotalString;
-				const FString MemTotalCommand = FString::Printf(TEXT("-s %s shell cat /proc/meminfo"), *NewDeviceInfo.SerialNumber);
+				const FString MemTotalCommand = FString::Printf(TEXT("-s %s shell cat /proc/meminfo"), *SerialNumber);
 				if (ExecuteAdbCommand(*MemTotalCommand, &MemTotalString, nullptr))
 				{
 					FString MemTotalLine = FindLineContaining(MemTotalString, TEXT("MemTotal"));
@@ -557,7 +601,7 @@ private:
 				FParse::Value(*DeviceString, TEXT("device:"), NewDeviceInfo.DeviceName);
 				if (NewDeviceInfo.DeviceName.IsEmpty())
 				{
-					FString DeviceCommand = FString::Printf(TEXT("-s %s %s ro.product.device"), *NewDeviceInfo.SerialNumber, *GetPropCommand);
+					FString DeviceCommand = FString::Printf(TEXT("-s %s %s ro.product.device"), *SerialNumber, *GetPropCommand);
 					FString RoProductDevice;
 					ExecuteAdbCommand(*DeviceCommand, &RoProductDevice, nullptr);
 					const TCHAR* Ptr = *RoProductDevice;
@@ -584,7 +628,7 @@ private:
 						const TCHAR* Ptr = *FwdString;
 						FString FwdSerialNumber, FwdHostPortString, FwdDevicePortString;
 						uint16 FwdHostPort, FwdDevicePort;
-						if (FParse::Token(Ptr, FwdSerialNumber, false) && FwdSerialNumber == NewDeviceInfo.SerialNumber &&
+						if (FParse::Token(Ptr, FwdSerialNumber, false) && FwdSerialNumber == SerialNumber &&
 							FParse::Token(Ptr, FwdHostPortString, false) && FParse::Value(*FwdHostPortString, TEXT("tcp:"), FwdHostPort) &&
 							FParse::Token(Ptr, FwdDevicePortString, false) && FParse::Value(*FwdDevicePortString, TEXT("tcp:"), FwdDevicePort) && FwdDevicePort == 6666)
 						{
@@ -612,7 +656,7 @@ private:
 							}
 						} while (!bFoundPort);
 
-						FString DeviceCommand = FString::Printf(TEXT("-s %s forward tcp:%d tcp:6666"), *NewDeviceInfo.SerialNumber, HostMessageBusPort);
+						FString DeviceCommand = FString::Printf(TEXT("-s %s forward tcp:%d tcp:6666"), *SerialNumber, HostMessageBusPort);
 						ExecuteAdbCommand(*DeviceCommand, nullptr, nullptr);
 						NewDeviceInfo.HostMessageBusPort = HostMessageBusPort;
 					}
@@ -621,24 +665,25 @@ private:
 				}
 
 				// Add reverse port forwarding
-				uint16 ReversePortMappings[] = {
+				uint16 ReversePortMappings[]
+				{
 					41899,	// Network file server, DEFAULT_TCP_FILE_SERVING_PORT in NetworkMessage.h
 					1981,	// Unreal Insights data collection, TraceInsightsModule.cpp
 #if UE_WITH_ZEN
-					UE::Zen::IsDefaultServicePresent() ? UE::Zen::GetDefaultServiceInstance().GetPort() : 0, // Zen Store, usually defaults to 8558
+					UE::Zen::IsDefaultServicePresent() ? UE::Zen::GetDefaultServiceInstance().GetPort() : uint16(0), // Zen Store, usually defaults to 8558
 #endif
 					0 // end of list
 				};
 
 				for (int32 Idx=0; ReversePortMappings[Idx] > 0; Idx++)
 				{
-					FString DeviceCommand = FString::Printf(TEXT("-s %s reverse tcp:%d tcp:%d"), *NewDeviceInfo.SerialNumber, ReversePortMappings[Idx], ReversePortMappings[Idx]);
+					FString DeviceCommand = FString::Printf(TEXT("-s %s reverse tcp:%d tcp:%d"), *SerialNumber, ReversePortMappings[Idx], ReversePortMappings[Idx]);
 					// It doesn't really matter if a mapping already exists. There is no listening local port so no contention between multiple editor instances
 					ExecuteAdbCommand(*DeviceCommand, nullptr, nullptr);
 				}
 
 				FString WindowDisplaysOutput;
-				const FString ExtensionsCommand = FString::Printf(TEXT("-s %s shell dumpsys window displays"), *NewDeviceInfo.SerialNumber);
+				const FString ExtensionsCommand = FString::Printf(TEXT("-s %s shell dumpsys window displays"), *SerialNumber);
 				if (!ExecuteAdbCommand(*ExtensionsCommand, &WindowDisplaysOutput, nullptr))
 				{
 					continue;
@@ -657,40 +702,70 @@ private:
 						NewDeviceInfo.InsetsBottom = FCString::Atof(*RegexMatcher.GetCaptureGroup(4)) / NewDeviceInfo.ResolutionY;
 					}
 				}
+
+				NewDeviceInfo.SerialNumber = MoveTemp(SerialNumber);
+			
+				{
+					FScopeLock ScopeLock(DeviceMapLock);
+
+					Device = &DeviceMap.Add(NewDeviceInfo.DeviceId);
+					*Device = MoveTemp(NewDeviceInfo);
+				}
 			}
+			
+			CurrentlyConnectedDevices.Add(Device->DeviceId);
+		}
 
-			// add the device to the map
+		for (FString& AvdName : AvdNames)
+		{
+			FString DeviceId = TEXT("avd-") + AvdName;
+			FAndroidDeviceInfo* Device = DeviceMap.Find(DeviceId);
+			
+			if (!Device || !Device->SerialNumber.IsEmpty())
 			{
-				FScopeLock ScopeLock(DeviceMapLock);
+				FAndroidDeviceInfo NewDeviceInfo
+				{
+					.DeviceId = MoveTemp(DeviceId),
+					.AvdName = MoveTemp(AvdName),
+					.GLESExtensions = TEXT("GL_KHR_texture_compression_astc_ldr"),
+					.GLESVersion = 0x30001,
+					.bAuthorizedDevice = true,
+					.VulkanVersion = TEXT("1.1.0")
+				};
+				
+				{
+					FScopeLock ScopeLock(DeviceMapLock);
 
-				FAndroidDeviceInfo& SavedDeviceInfo = DeviceMap.Add(NewDeviceInfo.SerialNumber);
-				SavedDeviceInfo = NewDeviceInfo;
+					Device = &DeviceMap.Add(NewDeviceInfo.DeviceId);
+					*Device = MoveTemp(NewDeviceInfo);
+				}
 			}
+			
+			CurrentlyConnectedDevices.Add(Device->DeviceId);
 		}
 
 		// loop through the previously connected devices list and remove any that aren't still connected from the updated DeviceMap
 		TArray<FString> DevicesToRemove;
 
-		for (auto It = DeviceMap.CreateConstIterator(); It; ++It)
+		for (const auto& Pair : DeviceMap)
 		{
-			if (!CurrentlyConnectedDevices.Contains(It.Key()))
+			if (!CurrentlyConnectedDevices.Contains(Pair.Key))
 			{
-				if (TcpMessagingModule && It.Value().HostMessageBusPort != 0)
+				if (TcpMessagingModule && Pair.Value.HostMessageBusPort != 0)
 				{
-					TcpMessagingModule->RemoveOutgoingConnection(FString::Printf(TEXT("127.0.0.1:%d"), It.Value().HostMessageBusPort));
+					TcpMessagingModule->RemoveOutgoingConnection(FString::Printf(TEXT("127.0.0.1:%d"), Pair.Value.HostMessageBusPort));
 				}
 
-				DevicesToRemove.Add(It.Key());
+				DevicesToRemove.Add(Pair.Key);
 			}
 		}
 
 		{
-			// enter the critical section and remove the devices from the map
 			FScopeLock ScopeLock(DeviceMapLock);
 
-			for (auto It = DevicesToRemove.CreateConstIterator(); It; ++It)
+			for (const FString& Device : DevicesToRemove)
 			{
-				DeviceMap.Remove(*It);
+				DeviceMap.Remove(Device);
 			}
 		}
 	}
@@ -699,6 +774,7 @@ private:
 
 	// path to the adb command
 	FString ADBPath;
+	FString AvdHomePath;
 	FString GetPropCommand;
 	bool bGetExtensionsViaSurfaceFlinger;
 
@@ -785,12 +861,14 @@ public:
 	virtual void UpdateADBPath() override
 	{
 		FScopeLock PathUpdateLock(&ADBPathCheckLock);
-		FString AndroidDirectory = FPlatformMisc::GetEnvironmentVariable(*SDKDirEnvVar);
 
-		ADBPath.Empty();
-		
+		FString AndroidHomeDirectory = FPlatformMisc::GetEnvironmentVariable(*SDKDirEnvVar);
+		FString AndroidUserHomeDirectory = FPlatformMisc::GetEnvironmentVariable(TEXT("ANDROID_USER_HOME"));
+		FString AndroidEmulatorHomeDirectory = FPlatformMisc::GetEnvironmentVariable(TEXT("ANDROID_EMULATOR_HOME"));
+		FString AndroidAvdHomeDirectory = FPlatformMisc::GetEnvironmentVariable(TEXT("ANDROID_AVD_HOME"));
+
 #if PLATFORM_MAC || PLATFORM_LINUX
-		if (AndroidDirectory.Len() == 0)
+		if (AndroidHomeDirectory.IsEmpty() || AndroidUserHomeDirectory.IsEmpty() || AndroidEmulatorHomeDirectory.IsEmpty() || AndroidAvdHomeDirectory.IsEmpty())
 		{
 #if PLATFORM_LINUX
 			// didn't find ANDROID_HOME, so parse the .bashrc file on Linux
@@ -812,24 +890,47 @@ public:
 				FString(ANSI_TO_TCHAR(AnsiContents)).ParseIntoArrayLines(Lines);
 				FMemory::Free(AnsiContents);
 
-				for (int32 Index = Lines.Num()-1; Index >=0; Index--)
+				auto UpdateEnvironmentVariable = [](const FString& Line, const TCHAR* Key, FString& Value)
 				{
-					if (AndroidDirectory.Len() == 0 && Lines[Index].StartsWith(FString::Printf(TEXT("export %s="), *SDKDirEnvVar)))
+					if (Value.IsEmpty() && Line.StartsWith(FString::Printf(TEXT("export %s="), Key)))
 					{
 						FString Directory;
-						Lines[Index].Split(TEXT("="), NULL, &Directory);
-						Directory = Directory.Replace(TEXT("\""), TEXT(""));
-						AndroidDirectory = Directory;
-						setenv(TCHAR_TO_ANSI(*SDKDirEnvVar), TCHAR_TO_ANSI(*AndroidDirectory), 1);
+						Line.Split(TEXT("="), NULL, &Directory);
+						Directory.ReplaceInline(TEXT("\""), TEXT(""));
+						Value = MoveTemp(Directory);
+						setenv(TCHAR_TO_ANSI(*Key), TCHAR_TO_ANSI(*Value), 1);
+						return true;
+					}
+
+					return false;
+				};
+
+				for (int32 Index = Lines.Num()-1; Index >=0; Index--)
+				{
+					if (UpdateEnvironmentVariable(Lines[Index], *SDKDirEnvVar, AndroidHomeDirectory))
+					{
+						continue;
+					}
+					else if (UpdateEnvironmentVariable(Lines[Index], TEXT("ANDROID_USER_HOME"), AndroidUserHomeDirectory))
+					{
+						continue;
+					}
+					else if (UpdateEnvironmentVariable(Lines[Index], TEXT("ANDROID_EMULATOR_HOME"), AndroidEmulatorHomeDirectory))
+					{
+						continue;
+					}
+					else if (UpdateEnvironmentVariable(Lines[Index], TEXT("ANDROID_AVD_HOME"), AndroidAvdHomeDirectory))
+					{
+						continue;
 					}
 				}
 			}
 		}
 #endif
 
-		if (AndroidDirectory.Len() > 0)
+		if (!AndroidHomeDirectory.IsEmpty())
 		{
-			ADBPath = FPaths::Combine(*AndroidDirectory, SDKRelativeExePath);
+			ADBPath = FPaths::Combine(*AndroidHomeDirectory, SDKRelativeExePath);
 
 			// if it doesn't exist then just clear the path as we might set it later
 			if (!FPaths::FileExists(ADBPath))
@@ -837,7 +938,29 @@ public:
 				ADBPath.Empty();
 			}
 		}
-		DetectionThreadRunnable->UpdateADBPath(ADBPath, GetPropCommand, bGetExtensionsViaSurfaceFlinger);
+		else
+		{
+			ADBPath.Empty();
+		}
+
+		if (AndroidAvdHomeDirectory.IsEmpty())
+		{
+			if (AndroidEmulatorHomeDirectory.IsEmpty())
+			{
+				if (AndroidUserHomeDirectory.IsEmpty())
+				{
+					AndroidUserHomeDirectory = FPaths::Combine(FPlatformProcess::UserHomeDir(), TEXT(".android"));
+				}
+				
+				AndroidEmulatorHomeDirectory = MoveTemp(AndroidUserHomeDirectory);
+			}
+			
+			AndroidAvdHomeDirectory = FPaths::Combine(AndroidEmulatorHomeDirectory, TEXT("avd"));
+		}
+
+		AvdHomePath = MoveTemp(AndroidAvdHomeDirectory);
+		
+		DetectionThreadRunnable->UpdatePaths(ADBPath, AvdHomePath, GetPropCommand, bGetExtensionsViaSurfaceFlinger);
 	}
 
 	virtual void ExportDeviceProfile(const FString& OutPath, const FString& DeviceName) override
@@ -928,6 +1051,7 @@ private:
 
 	FString SDKDirEnvVar;
 	FString SDKRelativeExePath;
+	FString AvdHomePath;
 	FString GetPropCommand;
 	bool bGetExtensionsViaSurfaceFlinger;
 
diff --git a/Engine/Source/Developer/Android/AndroidDeviceDetection/Public/Interfaces/IAndroidDeviceDetection.h b/Engine/Source/Developer/Android/AndroidDeviceDetection/Public/Interfaces/IAndroidDeviceDetection.h
--- a/Engine/Source/Developer/Android/AndroidDeviceDetection/Public/Interfaces/IAndroidDeviceDetection.h
+++ b/Engine/Source/Developer/Android/AndroidDeviceDetection/Public/Interfaces/IAndroidDeviceDetection.h
@@ -9,10 +9,14 @@
 #include "Containers/UnrealString.h"
 #include "HAL/CriticalSection.h"
 
-template<typename KeyType,typename ValueType,typename SetAllocator ,typename KeyFuncs > class TMap;
+template<typename KeyType, typename ValueType, typename SetAllocator, typename KeyFuncs> class TMap;
 
 struct FAndroidDeviceInfo
 {
+	FString DeviceId;
+	
+	FString AvdName;
+	
 	// Device serial number, used to route ADB commands to a specific device
 	FString SerialNumber;
 
@@ -26,19 +30,19 @@ struct FAndroidDeviceInfo
 	FString HumanAndroidVersion;
 
 	// Android SDK version supported by the device (ro.build.version.sdk - note: deprecated in 4 according to docs, but version 4 devices return an empty string when querying the 'replacement' SDK_INT)
-	int32 SDKVersion;
+	int32 SDKVersion = INDEX_NONE;
 
 	// List of supported OpenGL extensions (retrieved via SurfaceFlinger)
 	FString GLESExtensions;
 
 	// Supported GLES version (ro.opengles.version)
-	int32 GLESVersion;
+	int32 GLESVersion = INDEX_NONE;
 
 	// Is the device authorized for USB communication?  If not, then none of the other properties besides the serial number will be valid
-	bool bAuthorizedDevice;
+	bool bAuthorizedDevice = true;
 
 	// TCP port number on our local host forwarded over adb to the device
-	uint16 HostMessageBusPort;
+	uint16 HostMessageBusPort = 0;
 
 	// Holds pixel per inch value.
 	int32 DeviceDPI = 0;
@@ -73,14 +77,6 @@ struct FAndroidDeviceInfo
 
 	// device BuildNumber field
 	FString BuildNumber;
-
-	FAndroidDeviceInfo()
-		: SDKVersion(INDEX_NONE)
-		, GLESVersion(INDEX_NONE)
-		, bAuthorizedDevice(true)
-		, HostMessageBusPort(0)
-	{
-	}
 };
 
 
diff --git a/Engine/Source/Developer/Android/AndroidTargetPlatformControls/Private/AndroidTargetDevice.h b/Engine/Source/Developer/Android/AndroidTargetPlatformControls/Private/AndroidTargetDevice.h
--- a/Engine/Source/Developer/Android/AndroidTargetPlatformControls/Private/AndroidTargetDevice.h
+++ b/Engine/Source/Developer/Android/AndroidTargetPlatformControls/Private/AndroidTargetDevice.h
@@ -50,19 +50,20 @@ public:
 	 * @param InSerialNumber - The ADB serial number of the target device.
 	 * @param InAndroidVariant - The variant of the Android platform, i.e. ETC2, DXT or ASTC.
 	 */
-	FAndroidTargetDevice(const ITargetPlatformControls& InTargetPlatformControls, const FString& InSerialNumber, const FString& InAndroidVariant)
-		: AndroidVariant(InAndroidVariant)
-		, bConnected(false)
+	FAndroidTargetDevice(const ITargetPlatformControls& InTargetPlatformControls, const FString& InDeviceId, const FString& InAndroidVariant)
+		: DeviceId{InDeviceId}
+		, AndroidVariant(InAndroidVariant)
 		, bIsDeviceAuthorized(false)
 		, AndroidSDKVersion(INDEX_NONE)
-		, DeviceName(InSerialNumber)
-		, Model(InSerialNumber)
-		, SerialNumber(InSerialNumber)
 		, TargetPlatformControls(InTargetPlatformControls)
 	{ }
 
 public:
-
+	const FString& GetDeviceId() const
+	{
+		return DeviceId;
+	}
+	
 	/**
 	 * Sets the device's connection state.
 	 *
@@ -95,6 +96,16 @@ public:
 		AndroidVersionString = InReleaseVersion;
 	}
 
+	/**
+	 * Sets the name.
+	 *
+	 * @param InName - The name to set.
+	 */
+	void SetName(FString InName)
+	{
+		Name = MoveTemp(InName);
+	}
+
 	/**
 	 * Sets the device name.
 	 *
@@ -108,14 +119,24 @@ public:
 	/**
 	 * Sets the device name.
 	 *
-	 * @param InDeviceName - The device name to set.
+	 * @param InModel - The device name to set.
 	 */
-	void SetModel(const FString& InNodel)
+	void SetModel(const FString& InModel)
 	{
-		Model = InNodel;
+		Model = InModel;
 	}
 
-	FString GetSerialNumber() const
+	/**
+	 * Sets the serial number
+	 *
+	 * @param InSerialNumber - The serial number to set.
+	 */
+	void SetSerialNumber(FString InSerialNumber)
+	{
+		SerialNumber = MoveTemp(InSerialNumber);
+	}
+
+	const FString& GetSerialNumber() const
 	{
 		return SerialNumber;
 	}
@@ -140,13 +161,13 @@ public:
 
 	virtual FTargetDeviceId GetId() const override
 	{
-		return FTargetDeviceId(TargetPlatformControls.PlatformName(), SerialNumber);
+		return FTargetDeviceId(TargetPlatformControls.PlatformName(), DeviceId);
 	}
 
 	virtual FString GetName() const override
 	{
 		// we need a unique name for all devices, so use human usable model name and the unique id
-		return FString::Printf(TEXT("%s (%s)"), *Model, *SerialNumber);
+		return Name;
 	}
 
 	virtual FString GetModelId() const override
@@ -187,13 +208,9 @@ public:
 		return bIsDeviceAuthorized;
 	}
 
+	virtual bool PowerOn() override;
 	virtual bool PowerOff(bool Force) override;
 
-	virtual bool PowerOn() override
-	{
-		return true;
-	}
-
 	// Return true if the devices can be grouped in an aggregate (All_<platform>_devices_on_<host>) proxy
 	virtual bool IsPlatformAggregated() const override
 	{
@@ -239,6 +256,8 @@ protected:
 
 protected:
 
+	FString DeviceId;
+	
 	// The variant of the Android platform, i.e. ETC2, DXT or ASTC.
 	FString AndroidVariant;
 
@@ -254,6 +273,9 @@ protected:
 	// Holds the Android Release version string (e.g., "2.3" or "4.2.2")
 	FString AndroidVersionString;
 
+	// Holds the name.
+	FString Name;
+	
 	// Holds the device name.
 	FString DeviceName;
 
diff --git a/Engine/Source/Developer/Android/AndroidTargetPlatformControls/Private/AndroidTargetDevice.inl b/Engine/Source/Developer/Android/AndroidTargetPlatformControls/Private/AndroidTargetDevice.inl
--- a/Engine/Source/Developer/Android/AndroidTargetPlatformControls/Private/AndroidTargetDevice.inl
+++ b/Engine/Source/Developer/Android/AndroidTargetPlatformControls/Private/AndroidTargetDevice.inl
@@ -42,14 +42,27 @@ inline bool FAndroidTargetDevice::Reboot( bool bReconnect )
 	return true;
 }
 
-inline bool FAndroidTargetDevice::PowerOff( bool Force )
+inline bool FAndroidTargetDevice::PowerOn()
 {
-	if (!ExecuteAdbCommand(TEXT("reboot --poweroff"), NULL, NULL))
+	check(DeviceId.StartsWith(TEXT("avd-")));
+
+	if (FString AndroidHomeDirectory = FPlatformMisc::GetEnvironmentVariable(TEXT("ANDROID_HOME")); !AndroidHomeDirectory.IsEmpty())
 	{
-		return false;
+		FProcHandle Process = FPlatformProcess::CreateProc(*FPaths::Combine(AndroidHomeDirectory, TEXT("emulator"), TEXT("emulator")), *(TEXT("-avd ") + DeviceId.RightChop(4)), false, false, true, nullptr, 0, nullptr, nullptr);
+
+		if (Process.IsValid())
+		{
+			FPlatformProcess::CloseProc(Process);
+			return true;
+		}
 	}
 
-	return true;
+	return false;
+}
+
+inline bool FAndroidTargetDevice::PowerOff(bool Force)
+{
+	return ExecuteAdbCommand(TEXT("reboot -p"), nullptr, nullptr);
 }
 
 inline FString FAndroidTargetDevice::GetAllDevicesName() const
@@ -64,18 +77,18 @@ inline bool FAndroidTargetDevice::TerminateLaunchedProcess(const FString& Proces
 	return ExecuteAdbCommand(AdbCommand, nullptr, nullptr);
 }
 
-inline bool FAndroidTargetDevice::SupportsFeature( ETargetDeviceFeatures Feature ) const
+inline bool FAndroidTargetDevice::SupportsFeature(ETargetDeviceFeatures Feature) const
 {
 	switch (Feature)
 	{
 	case ETargetDeviceFeatures::PowerOff:
-		return true;
+		return bIsDeviceAuthorized;
 
 	case ETargetDeviceFeatures::PowerOn:
-		return false;
+		return SerialNumber.IsEmpty() && DeviceId.StartsWith(TEXT("avd-"));
 
 	case ETargetDeviceFeatures::Reboot:
-		return true;
+		return bIsDeviceAuthorized;
 
 	default:
 		return false;
diff --git a/Engine/Source/Developer/Android/AndroidTargetPlatformControls/Private/AndroidTargetPlatformControls.cpp b/Engine/Source/Developer/Android/AndroidTargetPlatformControls/Private/AndroidTargetPlatformControls.cpp
--- a/Engine/Source/Developer/Android/AndroidTargetPlatformControls/Private/AndroidTargetPlatformControls.cpp
+++ b/Engine/Source/Developer/Android/AndroidTargetPlatformControls/Private/AndroidTargetPlatformControls.cpp
@@ -385,33 +385,28 @@ bool FAndroidTargetPlatformControls::HandleTicker( float DeltaTime )
 		checkf(DeviceDetection != nullptr, TEXT("A target platform didn't create a device detection object in InitializeDeviceDetection()!"));
 	}
 
-	TArray<FString> ConnectedDeviceIds;
+	TArray<FStringView> ConnectedDeviceIds;
 
 	{
 		FScopeLock ScopeLock(DeviceDetection->GetDeviceMapLock());
 
-		auto DeviceIt = DeviceDetection->GetDeviceMap().CreateConstIterator();
-
-		for (; DeviceIt; ++DeviceIt)
+		for (const auto& Pair : DeviceDetection->GetDeviceMap())
 		{
-			ConnectedDeviceIds.Add(DeviceIt.Key());
-
-			const FAndroidDeviceInfo& DeviceInfo = DeviceIt.Value();
+			const FAndroidDeviceInfo& DeviceInfo = Pair.Value;
 
 			// see if this device is already known
-			if (Devices.Contains(DeviceIt.Key()))
+			if (FAndroidTargetDevicePtr TestDevice = Devices.FindRef(Pair.Key))
 			{
-				FAndroidTargetDevicePtr TestDevice = Devices[DeviceIt.Key()];
-
 				// ignore if authorization didn't change
-				if (DeviceInfo.bAuthorizedDevice == TestDevice->IsAuthorized())
+				if (DeviceInfo.bAuthorizedDevice == TestDevice->IsAuthorized() && DeviceInfo.SerialNumber == TestDevice->GetSerialNumber())
 				{
+					ConnectedDeviceIds.Add(TestDevice->GetDeviceId());
 					continue;
 				}
 
 				// remove it to add again
 				TestDevice->SetConnected(false);
-				Devices.Remove(DeviceIt.Key());
+				Devices.Remove(Pair.Key);
 
 				OnDeviceLost().Broadcast(TestDevice.ToSharedRef());
 			}
@@ -423,17 +418,39 @@ bool FAndroidTargetPlatformControls::HandleTicker( float DeltaTime )
 			}
 
 			// create target device
-			FAndroidTargetDevicePtr& Device = Devices.Add(DeviceInfo.SerialNumber);
+			FAndroidTargetDevicePtr& Device = Devices.Add(Pair.Key);
 
-			Device = CreateTargetDevice(*this, DeviceInfo.SerialNumber, GetAndroidVariantName());
+			Device = CreateTargetDevice(*this, DeviceInfo.DeviceId, GetAndroidVariantName());
+
+			// we need a unique name for all devices, so use human usable display name and the unique id
+			if (!DeviceInfo.SerialNumber.IsEmpty())
+			{
+				if (!DeviceInfo.Model.IsEmpty())
+				{
+					Device->SetName((!DeviceInfo.AvdName.IsEmpty() ? DeviceInfo.AvdName : DeviceInfo.Model) + TEXT(" (") + DeviceInfo.SerialNumber + TEXT(")"));
+				}
+				else
+				{
+					Device->SetName(DeviceInfo.SerialNumber + TEXT(" [UNAUTHORIZED]"));
+				}
+				
+				Device->SetConnected(true);
+			}
+			else
+			{
+				Device->SetName(DeviceInfo.AvdName + TEXT(" [OFFLINE]"));
+				Device->SetConnected(false);
+			}
 
-			Device->SetConnected(true);
 			Device->SetModel(DeviceInfo.Model);
 			Device->SetDeviceName(DeviceInfo.DeviceName);
 			Device->SetAuthorized(DeviceInfo.bAuthorizedDevice);
 			Device->SetVersions(DeviceInfo.SDKVersion, DeviceInfo.HumanAndroidVersion);
+			Device->SetSerialNumber(DeviceInfo.SerialNumber);
 
 			ITargetPlatformControls::OnDeviceDiscovered().Broadcast(Device.ToSharedRef());
+			
+			ConnectedDeviceIds.Add(Device->GetDeviceId());
 		}
 	}
 
@@ -456,12 +473,12 @@ bool FAndroidTargetPlatformControls::HandleTicker( float DeltaTime )
 
 FAndroidTargetDeviceRef FAndroidTargetPlatformControls::CreateNewDevice(const FAndroidDeviceInfo &DeviceInfo)
 {
-	return MakeShareable(new FAndroidTargetDevice(*this, DeviceInfo.SerialNumber, GetAndroidVariantName()));
+	return MakeShareable(new FAndroidTargetDevice(*this, DeviceInfo.DeviceId, GetAndroidVariantName()));
 }
 
-FAndroidTargetDevicePtr FAndroidTargetPlatformControls::CreateTargetDevice(const ITargetPlatformControls& InTargetPlatform, const FString& InSerialNumber, const FString& InAndroidVariant) const
+FAndroidTargetDevicePtr FAndroidTargetPlatformControls::CreateTargetDevice(const ITargetPlatformControls& InTargetPlatform, const FString& InDeviceId, const FString& InAndroidVariant) const
 {
-	return MakeShareable(new FAndroidTargetDevice(InTargetPlatform, InSerialNumber, InAndroidVariant));
+	return MakeShareable(new FAndroidTargetDevice(InTargetPlatform, InDeviceId, InAndroidVariant));
 }
 
 #if WITH_ENGINE
diff --git a/Engine/Source/Developer/Android/AndroidTargetPlatformControls/Private/AndroidTargetPlatformControls.h b/Engine/Source/Developer/Android/AndroidTargetPlatformControls/Private/AndroidTargetPlatformControls.h
--- a/Engine/Source/Developer/Android/AndroidTargetPlatformControls/Private/AndroidTargetPlatformControls.h
+++ b/Engine/Source/Developer/Android/AndroidTargetPlatformControls/Private/AndroidTargetPlatformControls.h
@@ -122,7 +122,7 @@ public:
 
 protected:
 
-	virtual FAndroidTargetDevicePtr CreateTargetDevice(const ITargetPlatformControls& InTargetPlatform, const FString& InSerialNumber, const FString& InAndroidVariant) const;
+	virtual FAndroidTargetDevicePtr CreateTargetDevice(const ITargetPlatformControls& InTargetPlatform, const FString& InDeviceId, const FString& InAndroidVariant) const;
 
 protected:
 	/**
diff --git a/Engine/Source/Editor/TurnkeySupport/Private/TurnkeySupportModule.cpp b/Engine/Source/Editor/TurnkeySupport/Private/TurnkeySupportModule.cpp
--- a/Engine/Source/Editor/TurnkeySupport/Private/TurnkeySupportModule.cpp
+++ b/Engine/Source/Editor/TurnkeySupport/Private/TurnkeySupportModule.cpp
@@ -1789,6 +1789,11 @@ static void GenerateDeviceProxyMenuParams(TSharedPtr<ITargetDeviceProxy> DeviceP
 	OutAction = FUIAction(
 		FExecuteAction::CreateLambda([DeviceProxy, PlatformName, ExternalOnClickDelegate]()
 			{
+				if (!DeviceProxy->IsConnected())
+				{
+					DeviceProxy->PowerOn();
+				}
+			
 				// only game and client devices are supported for launch on but devices only signal if they target client builds by their name e.g. "WindowsClient"
 				// if the user has a EBuildTargetType::Client target selected we will launch on a client device, otherwise we fall back to the default game device
 
diff --git a/Engine/Source/Programs/AutomationTool/Android/AndroidPlatform.Automation.cs b/Engine/Source/Programs/AutomationTool/Android/AndroidPlatform.Automation.cs
--- a/Engine/Source/Programs/AutomationTool/Android/AndroidPlatform.Automation.cs
+++ b/Engine/Source/Programs/AutomationTool/Android/AndroidPlatform.Automation.cs
@@ -4,6 +4,7 @@ using System.Collections.Generic;
 using System.Linq;
 using System.IO;
 using System.Diagnostics;
+using System.Text.Json;
 using System.Threading;
 using AutomationTool;
 using UnrealBuildTool;
@@ -264,7 +265,38 @@ public class AndroidPlatform : Platform
 		int ExitCode = TurnkeyContext.RunExternalCommand(Command, Params, bRequiresPrivilegeElevation: false, bUnattended, bCreateWindow: true);
 		return ExitCode == 0;
 	}
+	
+	public override bool PostSDKSetup(ITurnkeyContext TurnkeyContext, bool bUnattended)
+	{
+		FileReference EmulatorConfigFile = FileReference.Combine(Unreal.EngineDirectory, "Platforms", "Android", "Config", "AndroidEmulator.json");
+		
+		if (!FileReference.Exists(EmulatorConfigFile))
+		{
+			EmulatorConfigFile = FileReference.Combine(Unreal.EngineDirectory, "Config", "Android", "AndroidEmulator.json");
+			
+			if (!FileReference.Exists(EmulatorConfigFile))
+			{
+				return true;
+			}
+		}
+
+		Dictionary<string, string> EmulatorConfig = JsonSerializer.Deserialize<Dictionary<string, string>>(FileReference.ReadAllText(EmulatorConfigFile), new JsonSerializerOptions
+		{
+			AllowTrailingCommas = true,
+			ReadCommentHandling = JsonCommentHandling.Skip
+		});
+		
+		string CommandExtension = HostPlatform.Platform == UnrealTargetPlatform.Win64 ? "bat" : HostPlatform.Platform == UnrealTargetPlatform.Mac ? "command" : "sh";
+		
+		// run the SetupAndroidEmulator in the engine, not coming from a normal FileSource
+		string Command = $"$(EngineDir)/Extras/Android/SetupAndroidEmulator.{CommandExtension}";
+		string Params = $"{EmulatorConfig["AvdDevice"]} {EmulatorConfig["DeviceApiVersion"]} {EmulatorConfig["DeviceDataSize"]} {EmulatorConfig["DeviceRamSize"]} -noninteractive";
+		
+		// because this may bring up a license acceptance message that needs the user to respond, so we make a new window
+		int ExitCode = TurnkeyContext.RunExternalCommand(Command, Params, bRequiresPrivilegeElevation: false, bUnattended, bCreateWindow: true);
 
+		return ExitCode == 0;
+	}
 
 	private static string GetAndroidStudioExe()
 	{
@@ -502,6 +534,11 @@ public class AndroidPlatform : Platform
 
 		return bHaveAndroidStudio;
 	}
+
+	static string GetAvdHomePath()
+	{
+		return Environment.GetEnvironmentVariable("ANDROID_AVD_HOME") ?? Path.Combine(Environment.GetEnvironmentVariable("ANDROID_EMULATOR_HOME") ?? Environment.GetEnvironmentVariable("ANDROID_USER_HOME") ?? Path.Combine(Environment.GetFolderPath(Environment.SpecialFolder.UserProfile), ".android"), "avd");
+	}
 	
 	public override DeviceInfo[] GetDevices()
 	{
@@ -532,8 +569,32 @@ public class AndroidPlatform : Platform
 		}
 		return Devices.ToArray();
 	}
-
-
+	
+	public override DeviceInfo GetDeviceByName(string DeviceName)
+	{
+		if (DeviceName.StartsWith("avd-"))
+		{
+			return Directory.EnumerateFiles(GetAvdHomePath(), "*.ini").Select(Filename =>
+			{
+				string AvdName = "avd-" + Path.GetFileNameWithoutExtension(Filename);
+				string Target = GetIniValue(Filename, "target");
+				const string TargetPrefix = "android-";
+				string SdkVersion = Target.StartsWith(TargetPrefix) ? Target[TargetPrefix.Length..] : "";
+			
+				return new DeviceInfo(TargetPlatformType)
+				{
+					Id = AvdName,
+					Name = AvdName,
+					bCanConnect = true,
+					SoftwareVersion = SdkVersion.Length > 0 ? SdkVersion : null
+				};
+			}).FirstOrDefault(Device => string.Compare(Device.Id, DeviceName, true) == 0);
+		}
+		else
+		{
+			return base.GetDeviceByName(DeviceName);
+		}
+	}
 
 	private static string GetSONameWithoutArchitecture(ProjectParams Params, string DecoratedExeName)
 	{
@@ -2427,6 +2488,8 @@ public class AndroidPlatform : Platform
 
 	public override bool RetrieveDeployedManifests(ProjectParams Params, DeploymentContext SC, string DeviceName, out List<string> UFSManifests, out List<string> NonUFSManifests)
 	{
+		DeviceName = GetSerialNumber(DeviceName);
+		
 		bool bAFSEnablePlugin;
 		string AFSToken;
 		bool bIsShipping;
@@ -2522,7 +2585,7 @@ public class AndroidPlatform : Platform
 		AndroidFileClient.OptimalADB adb = new AndroidFileClient.OptimalADB();
 		int AllowOverflowOBBLimit = AllowOverflowOBBFiles(SC);
 
-		foreach (var DeviceName in Params.DeviceNames)
+		foreach (string DeviceName in Params.DeviceNames.Select(DeviceName => GetSerialNumber(DeviceName)))
 		{
 			UnrealArch? DeviceArchitecture = GetBestDeviceArchitecture(Params, DeviceName);
 			string ApkName = GetFinalApkName(Params, SC.StageExecutables[0], true, DeviceArchitecture);
@@ -3134,8 +3197,8 @@ public class AndroidPlatform : Platform
 		bool bDontBundleLibrariesInAPK = GetDontBundleLibrariesInAPK(Params, SC, true);
 		int AllowOverflowOBBLimit = AllowOverflowOBBFiles(SC);
 
-		foreach (var DeviceName in Params.DeviceNames)
-        {
+		foreach (string DeviceName in Params.DeviceNames.Select(DeviceName => GetSerialNumber(DeviceName)))
+		{
             UnrealArch? DeviceArchitecture = GetBestDeviceArchitecture(Params, DeviceName);
             string ApkName = GetFinalApkName(Params, SC.StageExecutables[0], true, DeviceArchitecture);
             string FinalSOName = GetSOName(Params, SC.StageExecutables[0], DeviceArchitecture);
@@ -3981,6 +4044,97 @@ public class AndroidPlatform : Platform
 
 		return CachedAaptPath;
 	}
+	
+	static string GetIniValue(string Path, string Key)
+	{
+		try
+		{
+			using StreamReader Reader = File.OpenText(Path);
+					
+			for (;;)
+			{
+				string Line = Reader.ReadLine();
+
+				if (Line == null)
+				{
+					break;
+				}
+
+				string[] Parts = Line.Split('=', 2);
+
+				if (Parts[0].Trim() == Key)
+				{
+					return Parts[1].Trim();
+				}
+			}
+		}
+		catch
+		{
+		}
+
+		return null;
+	}
+
+	static string GetSerialNumber(string DeviceName)
+	{
+		if (DeviceName.StartsWith("avd-"))
+		{
+			string AvdName = DeviceName[4..];
+
+			(string, string) QuerySerialNumber()
+			{
+				return RunAdbCommand(null, "devices", null, ERunOptions.AppMustExist).Output.Split(Environment.NewLine, StringSplitOptions.RemoveEmptyEntries).Skip(1).Where(Line => Line.StartsWith("emulator-")).Select(Line => Line.Split('\t', 2)).Where(Parts =>
+				{
+					IProcessResult Result = RunAdbCommand(Parts[0], "emu avd name", null, ERunOptions.AppMustExist);
+					return Result.bExitCodeSuccess && Result.Output[0..Result.Output.IndexOf(Environment.NewLine, StringComparison.Ordinal)] == AvdName;
+				}).Select(Parts => (Parts[0], Parts[1])).FirstOrDefault();
+			}
+
+			(string SerialNumber, string Status) = QuerySerialNumber();
+
+			if (SerialNumber == null)
+			{
+				using Process Process = Process.Start(new ProcessStartInfo(Environment.ExpandEnvironmentVariables("%ANDROID_HOME%/emulator/emulator"), ["-avd", AvdName])	{ UseShellExecute = false, CreateNoWindow = true, WindowStyle = ProcessWindowStyle.Hidden })!;
+				
+				for (uint RetryCount = 60;; --RetryCount)
+				{
+					bool bExited = Process.WaitForExit(1000);
+						
+					(SerialNumber, Status) = QuerySerialNumber();
+
+					if (bExited || SerialNumber != null || RetryCount == 0)
+					{
+						break;
+					}
+				}
+			}
+			
+			string QueryStatus()
+			{
+				return RunAdbCommand(null, "devices", null, ERunOptions.AppMustExist).Output.Split(Environment.NewLine, StringSplitOptions.RemoveEmptyEntries).Skip(1).Select(Line => Line.Split('\t', 2)).Where(Parts => Parts[0] == SerialNumber).Select(Parts => Parts[1]).FirstOrDefault();
+			}
+			
+			for (uint RetryCount = 60; Status != "device"; --RetryCount)
+			{
+				Logger.LogInformation("Device status for '{DeviceName}': {Status}", DeviceName, Status);
+			
+				if (string.IsNullOrEmpty(Status) || RetryCount == 0)
+				{
+					throw new AutomationException(ExitCode.Error_UnknownLaunchFailure, "Unable to run because couldn't connect to device '{0}'", DeviceName);
+				}
+				
+				Thread.Sleep(1000);
+
+				Status = QueryStatus();
+			}
+
+			return SerialNumber;
+		}
+		else
+		{
+			return DeviceName;
+		}
+	}
 
 	private UnrealArch? GetBestDeviceArchitecture(ProjectParams Params, string DeviceName)
 	{
@@ -4081,7 +4235,7 @@ public class AndroidPlatform : Platform
 		string LogPath = Path.Combine(Params.BaseStageDirectory, "Android\\logs");
 		Directory.CreateDirectory(LogPath);
 
-		foreach (string DeviceName in Params.DeviceNames)
+		foreach (string DeviceName in Params.DeviceNames.Select(DeviceName => GetSerialNumber(DeviceName)))
 		{
 			//get the package name and save that
 			UnrealArch? DeviceArchitecture = GetBestDeviceArchitecture(Params, DeviceName);
diff --git a/Engine/Source/Programs/AutomationTool/AutomationUtils/Platform.cs b/Engine/Source/Programs/AutomationTool/AutomationUtils/Platform.cs
--- a/Engine/Source/Programs/AutomationTool/AutomationUtils/Platform.cs
+++ b/Engine/Source/Programs/AutomationTool/AutomationUtils/Platform.cs
@@ -266,6 +266,11 @@ namespace AutomationTool
 			return false;
 		}
 
+		public virtual bool PostSDKSetup(ITurnkeyContext TurnkeyContext, bool bUnattended)
+		{
+			return true;
+		}
+
 		/// <summary>
 		/// Return a list of versions that will be used to create "fake" FileSource objects which are used
 		/// for install Sdks where no file downloads are needed
diff --git a/Engine/Source/Programs/AutomationTool/Turnkey/FileSource.cs b/Engine/Source/Programs/AutomationTool/Turnkey/FileSource.cs
--- a/Engine/Source/Programs/AutomationTool/Turnkey/FileSource.cs
+++ b/Engine/Source/Programs/AutomationTool/Turnkey/FileSource.cs
@@ -417,7 +417,14 @@ namespace Turnkey
 			}
 
 			// let the platform decide how to install
-			return AutomationTool.Platform.GetPlatform(Platform).InstallSDK(TurnkeyUtils.CommandUtilHelper, TurnkeyContext, Device, bUnattended, bSdkAlreadyInstalled);
+			Platform AutomationPlatform = AutomationTool.Platform.GetPlatform(Platform);
+
+			if (!AutomationPlatform.InstallSDK(TurnkeyUtils.CommandUtilHelper, TurnkeyContext, Device, bUnattended, bSdkAlreadyInstalled))
+			{
+				return false;
+			}
+
+			return AutomationPlatform.PostSDKSetup(TurnkeyContext, bUnattended);
 		}
 
 
diff --git a/Engine/Source/Programs/AutomationTool/Turnkey/SdkInfo.cs b/Engine/Source/Programs/AutomationTool/Turnkey/SdkInfo.cs
--- a/Engine/Source/Programs/AutomationTool/Turnkey/SdkInfo.cs
+++ b/Engine/Source/Programs/AutomationTool/Turnkey/SdkInfo.cs
@@ -284,7 +284,7 @@ namespace Turnkey
 
 			UEBuildPlatformSDK.GetSDKForPlatform(Platform.ToString()).ReactivateAutoSDK();
 
-			return true;
+			return AutomationPlatform.PostSDKSetup(TurnkeyContext, bUnattended);
 		}
 	}
 }
diff --git a/Engine/Source/Runtime/Core/Private/Windows/WindowsPlatformProcess.cpp b/Engine/Source/Runtime/Core/Private/Windows/WindowsPlatformProcess.cpp
--- a/Engine/Source/Runtime/Core/Private/Windows/WindowsPlatformProcess.cpp
+++ b/Engine/Source/Runtime/Core/Private/Windows/WindowsPlatformProcess.cpp
@@ -1324,6 +1324,25 @@ const TCHAR* FWindowsPlatformProcess::UserSettingsDir()
 	return *WindowsUserSettingsDir;
 }
 
+const TCHAR *FWindowsPlatformProcess::UserHomeDir()
+{
+	static FString WindowsUserHomeDir = [&]() -> FString 
+	{
+		if (TCHAR* UserPath; SUCCEEDED(SHGetKnownFolderPath(FOLDERID_Profile, 0, nullptr, &UserPath)))
+		{
+			FString Directory = UserPath;
+			Directory.ReplaceInline(TEXT("\\"), TEXT("/"));
+			Directory.Append(TEXT("/"));
+			Directory.Shrink();
+			CoTaskMemFree(UserPath);
+			return Directory;
+		}
+
+		return {};
+	}();
+	return *WindowsUserHomeDir;
+}
+
 const TCHAR* FWindowsPlatformProcess::UserSettingsDirMediumIntegrity()
 {
 	static FString WindowsUserSettingsDir;
diff --git a/Engine/Source/Runtime/Core/Public/Windows/WindowsPlatformProcess.h b/Engine/Source/Runtime/Core/Public/Windows/WindowsPlatformProcess.h
--- a/Engine/Source/Runtime/Core/Public/Windows/WindowsPlatformProcess.h
+++ b/Engine/Source/Runtime/Core/Public/Windows/WindowsPlatformProcess.h
@@ -147,6 +147,7 @@ public:
 	static CORE_API const TCHAR* UserDir();
 	static CORE_API const TCHAR* UserTempDir();
 	static CORE_API const TCHAR* UserSettingsDir();
+	static CORE_API const TCHAR* UserHomeDir();
 	static CORE_API const TCHAR* UserSettingsDirMediumIntegrity();
 	static CORE_API const TCHAR* ApplicationSettingsDir();
 	static CORE_API FString GetApplicationSettingsDir(const ApplicationSettingsContext& Settings);
