Small FIx for linux/mac Editor build

#android

diff --git a/Engine/Source/Developer/Android/AndroidDeviceDetection/Private/AndroidDeviceDetectionModule.cpp b/Engine/Source/Developer/Android/AndroidDeviceDetection/Private/AndroidDeviceDetectionModule.cpp
--- a/Engine/Source/Developer/Android/AndroidDeviceDetection/Private/AndroidDeviceDetectionModule.cpp
+++ b/Engine/Source/Developer/Android/AndroidDeviceDetection/Private/AndroidDeviceDetectionModule.cpp
@@ -898,7 +898,7 @@ public:
 						Line.Split(TEXT("="), NULL, &Directory);
 						Directory.ReplaceInline(TEXT("\""), TEXT(""));
 						Value = MoveTemp(Directory);
-						setenv(TCHAR_TO_ANSI(*Key), TCHAR_TO_ANSI(*Value), 1);
+						setenv(TCHAR_TO_ANSI(Key), TCHAR_TO_ANSI(*Value), 1);
 						return true;
 					}
 
