Zen server plugin to forward zen server port to any connected android device that has apps requiring zen streaming
#rnx

diff --git a/Engine/Source/Developer/Android/AndroidZenServerPlugin/AndroidPortForwarder/AndroidPortForwarder.cs b/Engine/Source/Developer/Android/AndroidZenServerPlugin/AndroidPortForwarder/AndroidPortForwarder.cs
new file mode 100644
--- /dev/null
+++ b/Engine/Source/Developer/Android/AndroidZenServerPlugin/AndroidPortForwarder/AndroidPortForwarder.cs
@@ -0,0 +1,51 @@
+﻿// Copyright Epic Games, Inc. All Rights Reserved.
+
+namespace AndroidZenServerPlugin;
+
+public sealed class AndroidPortForwarder : IDisposable
+{
+	private CancellationTokenSource? TokenSource;
+	private AndroidPortForwarderAsync? ForwarderAsync;
+	private Task? StartMonitorTask;
+
+	public void StartMonitor(string ADBExecutable, uint Port)
+	{
+		Console.WriteLine("Forwarder start requested");
+
+		TokenSource = new CancellationTokenSource();
+
+		ForwarderAsync = new AndroidPortForwarderAsync();
+
+		// We don't wait for startup to finish to avoid blocking current thread
+		StartMonitorTask = ForwarderAsync.StartMonitorAsync(ADBExecutable, Port, TokenSource.Token);
+
+		Console.WriteLine($"Forwarder started");
+	}
+
+	public void Dispose()
+	{
+		Console.WriteLine($"Forwarder stop requested");
+
+		TokenSource?.Cancel();
+
+		try
+		{
+			StartMonitorTask?.Wait();
+		}
+		catch (AggregateException exception)
+		{
+			// ignore token cancellation
+			exception.Flatten().Handle(inner => inner is OperationCanceledException);
+		}
+
+		StartMonitorTask?.Dispose();
+		ForwarderAsync?.Dispose();
+		TokenSource?.Dispose();
+
+		TokenSource = null;
+		ForwarderAsync = null;
+		StartMonitorTask = null;
+
+		Console.WriteLine($"Forwarder stopped");
+	}
+}
diff --git a/Engine/Source/Developer/Android/AndroidZenServerPlugin/AndroidPortForwarder/AndroidPortForwarder.csproj b/Engine/Source/Developer/Android/AndroidZenServerPlugin/AndroidPortForwarder/AndroidPortForwarder.csproj
new file mode 100644
--- /dev/null
+++ b/Engine/Source/Developer/Android/AndroidZenServerPlugin/AndroidPortForwarder/AndroidPortForwarder.csproj
@@ -0,0 +1,28 @@
+﻿<Project Sdk="Microsoft.NET.Sdk">
+  <PropertyGroup>
+    <!--
+    <OutputType>Exe</OutputType>
+    <DefineConstants>TEST_CLI</DefineConstants>
+    -->
+    <TargetFramework>net8.0</TargetFramework>
+    <ImplicitUsings>enable</ImplicitUsings>
+    <Nullable>enable</Nullable>
+    <PublishAot>true</PublishAot>
+    <InvariantGlobalization>true</InvariantGlobalization>
+    <IsAotCompatible>true</IsAotCompatible>
+    <StripSymbols>false</StripSymbols>
+    <Optimize>true</Optimize>
+    <OptimizationPreference>Size</OptimizationPreference>
+    <RootNamespace>AndroidZenServerPlugin</RootNamespace>
+    <AllowUnsafeBlocks>true</AllowUnsafeBlocks>
+  </PropertyGroup>
+  <ItemGroup>
+    <PackageReference Include="AdvancedSharpAdbClient" Version="3.4.14" />
+  </ItemGroup>
+  <ItemGroup>
+    <NativeLibrary Include="$(SolutionDir)$(Platform)\$(Configuration)\ZenServerAdapter.lib" Condition="$(RuntimeIdentifier.StartsWith('win'))" />
+    <NativeLibrary Include="$(SolutionDir)$(Platform)\$(Configuration)\ZenServerAdapter.a" Condition="!$(RuntimeIdentifier.StartsWith('win'))" />
+    <LinkerArg Include="/INCLUDE:GetTransportPluginVersion" Condition="$(RuntimeIdentifier.StartsWith('win'))" />
+    <LinkerArg Include="/INCLUDE:CreateTransportPlugin" Condition="$(RuntimeIdentifier.StartsWith('win'))" />
+  </ItemGroup>
+</Project>
diff --git a/Engine/Source/Developer/Android/AndroidZenServerPlugin/AndroidPortForwarder/AndroidPortForwarderAsync.cs b/Engine/Source/Developer/Android/AndroidZenServerPlugin/AndroidPortForwarder/AndroidPortForwarderAsync.cs
new file mode 100644
--- /dev/null
+++ b/Engine/Source/Developer/Android/AndroidZenServerPlugin/AndroidPortForwarder/AndroidPortForwarderAsync.cs
@@ -0,0 +1,104 @@
+﻿// Copyright Epic Games, Inc. All Rights Reserved.
+
+namespace AndroidZenServerPlugin;
+
+using AdvancedSharpAdbClient;
+using AdvancedSharpAdbClient.DeviceCommands;
+using AdvancedSharpAdbClient.Models;
+
+public sealed class AndroidPortForwarderAsync : IDisposable
+{
+	private AdbClient? Client;
+	private DeviceMonitor? DeviceMonitor;
+
+	public async Task StartMonitorAsync(string ADBExecutable, uint Port, CancellationToken CancellationToken)
+	{
+		try
+		{
+			if (string.IsNullOrEmpty(ADBExecutable))
+			{
+				ADBExecutable = Environment.ExpandEnvironmentVariables("%ANDROID_HOME%/platform-tools/adb" + (OperatingSystem.IsWindows() ? ".exe" : ""));
+			}
+
+			await AdbServer.Instance.StartServerAsync(ADBExecutable, false, CancellationToken);
+
+			Console.WriteLine($"Forwarder server started at '{ADBExecutable}'");
+
+			Client = new AdbClient();
+			DeviceMonitor = new DeviceMonitor(Client.CreateAdbSocket());
+
+			DeviceMonitor.DeviceConnected += async (_, args) =>
+			{
+				if (args.Device.State == DeviceState.Online)
+				{
+					await ForwardPortsToDeviceAsync(Client, args.Device, Port, CancellationToken);
+				}
+			};
+			DeviceMonitor.DeviceChanged += async (_, args) =>
+			{
+				if (args.Device.State == DeviceState.Online)
+				{
+					await ForwardPortsToDeviceAsync(Client, args.Device, Port, CancellationToken);
+				}
+			};
+
+			await DeviceMonitor.StartAsync(CancellationToken);
+
+			Console.WriteLine($"Forwarder ready");
+		}
+		catch (System.Net.Sockets.SocketException e)
+		{
+			Console.WriteLine($"Forwarder failed to connect to adb due to '{e}'");
+		}
+	}
+
+	public void Dispose()
+	{
+		DeviceMonitor?.Dispose();
+
+		DeviceMonitor = null;
+		Client = null;
+	}
+
+	private static async Task ForwardPortsToDeviceAsync(IAdbClient Client, DeviceData Device, uint Port, CancellationToken CancellationToken)
+	{
+		if (await DoesDeviceNeedsZenStreamingAsync(Client, Device, CancellationToken))
+		{
+			string Forward = $"tcp:{Port}";
+			await Client.CreateReverseForwardAsync(Device, Forward, Forward, true, CancellationToken);
+
+			Console.WriteLine($"Android device '{Device.Serial}' '{Forward}' reverse forwarded.");
+		}
+		else
+		{
+			Console.WriteLine($"Android device '{Device.Serial}' has no zen enabled apps installed, skipping.");
+		}
+	}
+
+	private static async Task<bool> DoesDeviceNeedsZenStreamingAsync(IAdbClient Client, DeviceData Device, CancellationToken CancellationToken)
+	{
+		// We quickly determinate if any apps on a device needs streaming via
+		// probing for a special activity added to every app's AndroidManifest.xml that uses zen streaming.
+		const string TokenActivity = "com.epicgames.unreal.Zen.StreamingEnabled";
+
+		const string Cmd = $"pm query-activities --brief -a {TokenActivity}";
+
+		bool bAnyActivities = false;
+
+		await Client.ExecuteShellCommandAsync(Device, Cmd, output =>
+		{
+			// Output will contain "No activities found" or "%number% activities found:".
+			// The strings are hardcoded in frameworks/base/services/core/java/com/android/server/pm/PackageManagerShellCommand.java, see
+			// https://cs.android.com/android/platform/superproject/+/android15-qpr2-release:frameworks/base/services/core/java/com/android/server/pm/PackageManagerShellCommand.java;l=1442-1445
+			if (output.Contains("activities found"))
+			{
+				bAnyActivities = !output.Contains("No activities found");
+				return false;
+			}
+
+			return true;
+		}, CancellationToken);
+
+		return bAnyActivities;
+	}
+}
diff --git a/Engine/Source/Developer/Android/AndroidZenServerPlugin/AndroidPortForwarder/AndroidPortForwarderCLI.cs b/Engine/Source/Developer/Android/AndroidZenServerPlugin/AndroidPortForwarder/AndroidPortForwarderCLI.cs
new file mode 100644
--- /dev/null
+++ b/Engine/Source/Developer/Android/AndroidZenServerPlugin/AndroidPortForwarder/AndroidPortForwarderCLI.cs
@@ -0,0 +1,25 @@
+﻿#if TEST_CLI
+
+using AndroidZenServerPlugin;
+
+TextWriter OriginalOut = Console.Out;
+using ConsoleRedirector Redirector = new ConsoleRedirector((Level, Line) => OriginalOut.WriteLine($"[{Level}] {Line}"));
+
+Console.WriteLine("Android Port Forwarder CLI testbed");
+
+ManualResetEvent exitEvent = new ManualResetEvent(false);
+
+Console.CancelKeyPress += (_, eventArgs) =>
+{
+	Console.WriteLine("Exit requested");
+	eventArgs.Cancel = true;
+	exitEvent.Set();
+};
+
+using AndroidPortForwarder forwarder = new AndroidPortForwarder();
+
+forwarder.StartMonitor(string.Empty, 8558);
+
+exitEvent.WaitOne();
+
+#endif
\ No newline at end of file
diff --git a/Engine/Source/Developer/Android/AndroidZenServerPlugin/AndroidPortForwarder/ConsoleRedirector.cs b/Engine/Source/Developer/Android/AndroidZenServerPlugin/AndroidPortForwarder/ConsoleRedirector.cs
new file mode 100644
--- /dev/null
+++ b/Engine/Source/Developer/Android/AndroidZenServerPlugin/AndroidPortForwarder/ConsoleRedirector.cs
@@ -0,0 +1,95 @@
+﻿namespace AndroidZenServerPlugin;
+
+// Redirects console output, errors and unhandled exceptions to a delegate invoked per output line.
+internal sealed class ConsoleRedirector : IDisposable
+{
+	public enum LogLevel
+	{
+		Info,
+		Error
+	};
+
+	public delegate void OutputDelegate(LogLevel Level, string Line);
+
+	private TextWriter OriginalOut;
+	private TextWriter OriginalError;
+
+	public ConsoleRedirector(OutputDelegate Delegate)
+	{
+		OriginalOut = Console.Out;
+		OriginalError = Console.Error;
+
+		Console.SetOut(new DelegateWriter(Delegate, LogLevel.Info));
+		Console.SetError(new DelegateWriter(Delegate, LogLevel.Error));
+
+		AppDomain.CurrentDomain.UnhandledException += OnCurrentDomainOnUnhandledException;
+	}
+
+	public void Dispose()
+	{
+		AppDomain.CurrentDomain.UnhandledException -= OnCurrentDomainOnUnhandledException;
+
+		Console.SetOut(OriginalOut);
+		Console.SetError(OriginalError);
+	}
+
+	private void OnCurrentDomainOnUnhandledException(object _, UnhandledExceptionEventArgs args)
+	{
+		Console.Error.WriteLine($"Unhandled exception: {args.ExceptionObject}");
+		Console.Error.Flush();
+	}
+
+	private sealed class DelegateWriter : StringWriter
+	{
+		private readonly LogLevel LogLevel;
+		private readonly OutputDelegate Delegate;
+
+		public DelegateWriter(OutputDelegate InDelegate, LogLevel InLogLevel)
+		{
+			LogLevel = InLogLevel;
+			Delegate = InDelegate;
+		}
+
+		public override void Flush()
+		{
+			using var Reader = new StringReader(ToString());
+			GetStringBuilder().Clear();
+
+			string? Line;
+			while ((Line = Reader.ReadLine()) != null)
+			{
+				Delegate(LogLevel, Line);
+			}
+		}
+
+		private void FlushLine()
+		{
+			Delegate(LogLevel, ToString());
+			GetStringBuilder().Clear();
+		}
+
+		public override void Write(char[] buffer, int index, int count)
+		{
+			if (buffer == CoreNewLine && index == 0 && count == CoreNewLine.Length)
+			{
+				FlushLine();
+			}
+			else
+			{
+				base.Write(buffer, index, count);
+			}
+		}
+
+		public override void Write(string? value)
+		{
+			if (value == NewLine)
+			{
+				FlushLine();
+			}
+			else
+			{
+				base.Write(value);
+			}
+		}
+	}
+}
\ No newline at end of file
diff --git a/Engine/Source/Developer/Android/AndroidZenServerPlugin/AndroidPortForwarder/ZenServerAdapter.cs b/Engine/Source/Developer/Android/AndroidZenServerPlugin/AndroidPortForwarder/ZenServerAdapter.cs
new file mode 100644
--- /dev/null
+++ b/Engine/Source/Developer/Android/AndroidZenServerPlugin/AndroidPortForwarder/ZenServerAdapter.cs
@@ -0,0 +1,56 @@
+﻿// Copyright Epic Games, Inc. All Rights Reserved.
+
+namespace AndroidZenServerPlugin;
+
+using System.Runtime.InteropServices;
+
+public static class ZenServerAdapter
+{
+	private static ConsoleRedirector? ConsoleRedirector;
+	private static AndroidPortForwarder? AndroidPortForwarder;
+
+	[UnmanagedCallersOnly(EntryPoint = "StartAndroidPortForwarder")]
+	private static unsafe void StartAndroidPortForwarder(delegate*<void*, int, char*, void> LogCallback, void* UsrPtr, IntPtr ADBExecutablePtr, uint Port)
+	{
+		ConsoleRedirector?.Dispose();
+		ConsoleRedirector = new ConsoleRedirector((Level, Line) =>
+		{
+			IntPtr LinePtr = Marshal.StringToHGlobalAnsi(Line);
+			try
+			{
+				LogCallback(UsrPtr, LogLevelToInt(Level), (char*)LinePtr);
+			}
+			finally
+			{
+				Marshal.FreeHGlobal(LinePtr);
+			}
+		});
+
+		string ADBExecutable = Marshal.PtrToStringAnsi(ADBExecutablePtr) ?? string.Empty;
+
+		AndroidPortForwarder?.Dispose();
+		AndroidPortForwarder = new AndroidPortForwarder();
+
+		AndroidPortForwarder.StartMonitor(ADBExecutable, Port);
+	}
+
+	[UnmanagedCallersOnly(EntryPoint = "StopAndroidPortForwarder")]
+	private static void StopAndroidPortForwarder()
+	{
+		AndroidPortForwarder?.Dispose();
+		AndroidPortForwarder = null;
+
+		ConsoleRedirector?.Dispose();
+		ConsoleRedirector = null;
+	}
+
+	private static int LogLevelToInt(ConsoleRedirector.LogLevel Level)
+	{
+		return Level switch
+		{
+			ConsoleRedirector.LogLevel.Info => 0,
+			ConsoleRedirector.LogLevel.Error => 1,
+			_ => 0
+		};
+	}
+}
diff --git a/Engine/Source/Developer/Android/AndroidZenServerPlugin/ZenServerAdapter/ZenServerAdapter.cpp b/Engine/Source/Developer/Android/AndroidZenServerPlugin/ZenServerAdapter/ZenServerAdapter.cpp
new file mode 100644
--- /dev/null
+++ b/Engine/Source/Developer/Android/AndroidZenServerPlugin/ZenServerAdapter/ZenServerAdapter.cpp
@@ -0,0 +1,136 @@
+// Copyright Epic Games, Inc. All Rights Reserved.
+
+#include <transportplugin.h>
+#include <atomic>
+#include <cstring>
+#include <cstdlib>
+#include <cstdio>
+
+#define USE_LOGGER 0
+
+extern "C"
+{
+	typedef void (*LogCallback)(void* UsrPtr, int32_t Level, const char* Buffer);
+
+	// native -> managed API
+	void StartAndroidPortForwarder(LogCallback Callback, void* UsrPtr, const char* ADBExecutable, uint32_t Port);
+	void StopAndroidPortForwarder();
+}
+
+class ZenServerAdapter : public zen::TransportPlugin
+{
+public:
+#if USE_LOGGER
+	ZenServerAdapter(zen::TransportLogger* InLogger)
+		: Logger(InLogger)
+	{
+	}
+#else
+	ZenServerAdapter() = default;
+#endif
+	virtual ~ZenServerAdapter() = default;
+
+	uint32_t AddRef() const override
+	{
+		return const_cast<ZenServerAdapter*>(this)->ReferenceCount.fetch_add(1) + 1;
+	}
+
+	uint32_t Release() const override
+	{
+		const uint32_t RefCount = const_cast<ZenServerAdapter*>(this)->ReferenceCount.fetch_sub(1);
+		if (RefCount <= 1)
+		{
+			delete this;
+		}
+		return RefCount - 1;
+	}
+
+	void Configure(const char* OptionTag, const char* OptionValue) override
+	{
+		if (!strcmp(OptionTag, "port"))
+		{
+			char* EndPtr = nullptr;
+			const long PortValue = strtoul(OptionValue, &EndPtr, 10);
+			if (OptionValue != EndPtr)
+			{
+				Port = (uint32_t)PortValue;
+			}
+			else
+			{
+#if USE_LOGGER
+				Logger->LogMessage(zen::TransportLogger::LogLevel::Err, "Can't parse port value");
+#endif
+			}
+		}
+		else if (!strcmp(OptionTag, "adb_executable"))
+		{
+			snprintf(ADBExecutable, sizeof(ADBExecutable), "%s", OptionValue);
+		}
+	}
+
+	void Initialize(zen::TransportServer* ServerInterface) override
+	{
+		StartAndroidPortForwarder(
+		[](void* ThisPtr, int32_t Level, const char* Buffer)
+		{
+#if USE_LOGGER
+			static_cast<ZenServerAdapter*>(ThisPtr)->Logger->LogMessage(Level == 1 ? zen::TransportLogger::LogLevel::Err : zen::TransportLogger::LogLevel::Info, Buffer);
+#endif
+		},
+		this, ADBExecutable, Port);
+	}
+
+	void Shutdown() override
+	{
+		StopAndroidPortForwarder();
+	}
+
+	const char* GetDebugName() override
+	{
+		return "AndroidPortForwarder";
+	}
+
+	bool IsAvailable() override
+	{
+		return true;
+	}
+
+private:
+	std::atomic<uint32_t> ReferenceCount = {0};
+	uint32_t Port = 8558;
+	char ADBExecutable[1024] = {};
+#if USE_LOGGER
+	zen::TransportLogger* Logger = nullptr;
+#endif
+};
+
+#if defined(_MSC_VER)
+#	define DLL_TRANSPORT_API __declspec(dllexport)
+#else
+#	define DLL_TRANSPORT_API
+#endif
+
+DLL_TRANSPORT_API void GetTransportPluginVersion(uint32_t* OutApiVersion, uint32_t* OutPluginVersion)
+{
+	if (OutApiVersion != nullptr)
+	{
+		*OutApiVersion = zen::kTransportApiVersion;
+	}
+
+	if (OutPluginVersion != nullptr)
+	{
+		*OutPluginVersion = 1;
+	}
+}
+
+#if USE_LOGGER
+extern "C" DLL_TRANSPORT_API zen::TransportPlugin* CreateTransportPlugin(zen::TransportLogger* Logger)
+{
+	return new ZenServerAdapter(Logger);
+}
+#else
+extern "C" DLL_TRANSPORT_API zen::TransportPlugin* CreateTransportPlugin()
+{
+	return new ZenServerAdapter();
+}
+#endif
\ No newline at end of file
diff --git a/Engine/Source/Developer/Android/AndroidZenServerPlugin/readme.md b/Engine/Source/Developer/Android/AndroidZenServerPlugin/readme.md
new file mode 100644
--- /dev/null
+++ b/Engine/Source/Developer/Android/AndroidZenServerPlugin/readme.md
@@ -0,0 +1,4 @@
+Zen server plugin for Android device support
+Reverse-forwards Zen Server port to any connected android device
+
+Use `dotnet publish -r win-x64` to build the .dll
\ No newline at end of file
