Add copyright
#rnx

diff --git a/Engine/Source/Developer/Android/AndroidZenServerPlugin/AndroidPortForwarder/AndroidPortForwarderCLI.cs b/Engine/Source/Developer/Android/AndroidZenServerPlugin/AndroidPortForwarder/AndroidPortForwarderCLI.cs
--- a/Engine/Source/Developer/Android/AndroidZenServerPlugin/AndroidPortForwarder/AndroidPortForwarderCLI.cs
+++ b/Engine/Source/Developer/Android/AndroidZenServerPlugin/AndroidPortForwarder/AndroidPortForwarderCLI.cs
@@ -1,4 +1,6 @@
-﻿#if TEST_CLI
+﻿// Copyright Epic Games, Inc. All Rights Reserved.
+
+#if TEST_CLI
 
 using AndroidZenServerPlugin;
 
diff --git a/Engine/Source/Developer/Android/AndroidZenServerPlugin/AndroidPortForwarder/ConsoleRedirector.cs b/Engine/Source/Developer/Android/AndroidZenServerPlugin/AndroidPortForwarder/ConsoleRedirector.cs
--- a/Engine/Source/Developer/Android/AndroidZenServerPlugin/AndroidPortForwarder/ConsoleRedirector.cs
+++ b/Engine/Source/Developer/Android/AndroidZenServerPlugin/AndroidPortForwarder/ConsoleRedirector.cs
@@ -1,4 +1,6 @@
-﻿namespace AndroidZenServerPlugin;
+﻿// Copyright Epic Games, Inc. All Rights Reserved.
+
+namespace AndroidZenServerPlugin;
 
 // Redirects console output, errors and unhandled exceptions to a delegate invoked per output line.
 internal sealed class ConsoleRedirector : IDisposable
