Update android port forwarding plugin

#rnx

diff --git a/Engine/Source/Developer/Android/AndroidZenServerPlugin/ZenServerAdapter/ZenServerAdapter.cpp b/Engine/Source/Developer/Android/AndroidZenServerPlugin/ZenServerAdapter/ZenServerAdapter.cpp
--- a/Engine/Source/Developer/Android/AndroidZenServerPlugin/ZenServerAdapter/ZenServerAdapter.cpp
+++ b/Engine/Source/Developer/Android/AndroidZenServerPlugin/ZenServerAdapter/ZenServerAdapter.cpp
@@ -6,8 +6,6 @@
 #include <cstdlib>
 #include <cstdio>
 
-#define USE_LOGGER 0
-
 extern "C"
 {
 	typedef void (*LogCallback)(void* UsrPtr, int32_t Level, const char* Buffer);
@@ -20,14 +18,10 @@ extern "C"
 class ZenServerAdapter : public zen::TransportPlugin
 {
 public:
-#if USE_LOGGER
 	ZenServerAdapter(zen::TransportLogger* InLogger)
 		: Logger(InLogger)
 	{
 	}
-#else
-	ZenServerAdapter() = default;
-#endif
 	virtual ~ZenServerAdapter() = default;
 
 	uint32_t AddRef() const override
@@ -57,9 +51,7 @@ public:
 			}
 			else
 			{
-#if USE_LOGGER
 				Logger->LogMessage(zen::TransportLogger::LogLevel::Err, "Can't parse port value");
-#endif
 			}
 		}
 		else if (!strcmp(OptionTag, "adb_executable"))
@@ -73,9 +65,7 @@ public:
 		StartAndroidPortForwarder(
 		[](void* ThisPtr, int32_t Level, const char* Buffer)
 		{
-#if USE_LOGGER
 			static_cast<ZenServerAdapter*>(ThisPtr)->Logger->LogMessage(Level == 1 ? zen::TransportLogger::LogLevel::Err : zen::TransportLogger::LogLevel::Info, Buffer);
-#endif
 		},
 		this, ADBExecutable, Port);
 	}
@@ -99,9 +89,7 @@ private:
 	std::atomic<uint32_t> ReferenceCount = {0};
 	uint32_t Port = 8558;
 	char ADBExecutable[1024] = {};
-#if USE_LOGGER
 	zen::TransportLogger* Logger = nullptr;
-#endif
 };
 
 #if defined(_MSC_VER)
@@ -110,7 +98,7 @@ private:
 #	define DLL_TRANSPORT_API
 #endif
 
-DLL_TRANSPORT_API void GetTransportPluginVersion(uint32_t* OutApiVersion, uint32_t* OutPluginVersion)
+extern "C" DLL_TRANSPORT_API void GetTransportPluginVersion(uint32_t* OutApiVersion, uint32_t* OutPluginVersion)
 {
 	if (OutApiVersion != nullptr)
 	{
@@ -123,14 +111,7 @@ DLL_TRANSPORT_API void GetTransportPluginVersion(uint32_t* OutApiVersion, uint32
 	}
 }
 
-#if USE_LOGGER
 extern "C" DLL_TRANSPORT_API zen::TransportPlugin* CreateTransportPlugin(zen::TransportLogger* Logger)
 {
 	return new ZenServerAdapter(Logger);
 }
-#else
-extern "C" DLL_TRANSPORT_API zen::TransportPlugin* CreateTransportPlugin()
-{
-	return new ZenServerAdapter();
-}
-#endif
\ No newline at end of file
