Don't treat failure to launch adb.exe as error
#rnx

diff --git a/Engine/Source/Developer/Android/AndroidZenServerPlugin/AndroidPortForwarder/AndroidPortForwarderAsync.cs b/Engine/Source/Developer/Android/AndroidZenServerPlugin/AndroidPortForwarder/AndroidPortForwarderAsync.cs
--- a/Engine/Source/Developer/Android/AndroidZenServerPlugin/AndroidPortForwarder/AndroidPortForwarderAsync.cs
+++ b/Engine/Source/Developer/Android/AndroidZenServerPlugin/AndroidPortForwarder/AndroidPortForwarderAsync.cs
@@ -50,28 +50,52 @@ public sealed class AndroidPortForwarderAsync : IDisposable
 		{
 			Console.WriteLine($"Forwarder failed to connect to adb due to '{e}'");
 		}
+		catch (FileNotFoundException e)
+		{
+			Console.WriteLine($"Forwarder failed to connect to adb due to '{e}'");
+		}
+		catch (Exception e)
+		{
+			Console.Error.WriteLine($"Forwarder failed to start due to '{e}'");
+		}
 	}
 
 	public void Dispose()
 	{
-		DeviceMonitor?.Dispose();
-
-		DeviceMonitor = null;
-		Client = null;
+		try
+		{
+			DeviceMonitor?.Dispose();
+		}
+		catch (Exception e)
+		{
+			Console.Error.WriteLine($"Forwarder failed to dispose due to '{e}'");
+		}
+		finally
+		{
+			DeviceMonitor = null;
+			Client = null;
+		}
 	}
 
 	private static async Task ForwardPortsToDeviceAsync(IAdbClient Client, DeviceData Device, uint Port, CancellationToken CancellationToken)
 	{
-		if (await DoesDeviceNeedsZenStreamingAsync(Client, Device, CancellationToken))
+		try
 		{
-			string Forward = $"tcp:{Port}";
-			await Client.CreateReverseForwardAsync(Device, Forward, Forward, true, CancellationToken);
+			if (await DoesDeviceNeedsZenStreamingAsync(Client, Device, CancellationToken))
+			{
+				string Forward = $"tcp:{Port}";
+				await Client.CreateReverseForwardAsync(Device, Forward, Forward, true, CancellationToken);
 
-			Console.WriteLine($"Android device '{Device.Serial}' '{Forward}' reverse forwarded.");
+				Console.WriteLine($"Android device '{Device.Serial}' '{Forward}' reverse forwarded.");
+			}
+			else
+			{
+				Console.WriteLine($"Android device '{Device.Serial}' has no zen enabled apps installed, skipping.");
+			}
 		}
-		else
+		catch (Exception e)
 		{
-			Console.WriteLine($"Android device '{Device.Serial}' has no zen enabled apps installed, skipping.");
+			Console.Error.WriteLine($"Forwarder failed to forward ports due to '{e}'");
 		}
 	}
 
diff --git a/Engine/Source/Developer/Android/AndroidZenServerPlugin/ZenServerAdapter/ZenServerAdapter.cpp b/Engine/Source/Developer/Android/AndroidZenServerPlugin/ZenServerAdapter/ZenServerAdapter.cpp
--- a/Engine/Source/Developer/Android/AndroidZenServerPlugin/ZenServerAdapter/ZenServerAdapter.cpp
+++ b/Engine/Source/Developer/Android/AndroidZenServerPlugin/ZenServerAdapter/ZenServerAdapter.cpp
@@ -107,7 +107,7 @@ extern "C" DLL_TRANSPORT_API void GetTransportPluginVersion(uint32_t* OutApiVers
 
 	if (OutPluginVersion != nullptr)
 	{
-		*OutPluginVersion = 1;
+		*OutPluginVersion = 2;
 	}
 }
 
