Fixed Editor repeatedly establishing connections to authorized Android devices

#android
#tests tested that reconnection doesn't happen on polling of previously authorized devices

diff --git a/Engine/Source/Developer/Android/AndroidDeviceDetection/Private/AndroidDeviceDetectionModule.cpp b/Engine/Source/Developer/Android/AndroidDeviceDetection/Private/AndroidDeviceDetectionModule.cpp
--- a/Engine/Source/Developer/Android/AndroidDeviceDetection/Private/AndroidDeviceDetectionModule.cpp
+++ b/Engine/Source/Developer/Android/AndroidDeviceDetection/Private/AndroidDeviceDetectionModule.cpp
@@ -326,16 +326,14 @@ private:
 				Status.LeftInline(TabIndex2);
 			}
 
-			bool bUnauthorized = Status == TEXT("unauthorized");
+			bool bAuthorized = Status != TEXT("unauthorized");
 
-			if (!bUnauthorized && Status != TEXT("device"))
+			if (bAuthorized && Status != TEXT("device"))
 			{
 				continue;
 			}
 
 			FString SerialNumber = DeviceString.Left(TabIndex), DeviceId, AvdName;
-			FAndroidDeviceInfo* Device;
-			FAndroidDeviceInfo NewDeviceInfo;
 
 			// Find the AVD name
 			if (FString AvdNameOutput; SerialNumber.StartsWith("emulator-") && ExecuteAdbCommand(*(TEXT("-s ") + SerialNumber + TEXT(" emu avd name")), &AvdNameOutput, nullptr))
@@ -351,397 +349,384 @@ private:
 				DeviceId = SerialNumber;
 			}
 			
-			if (bUnauthorized)
-			{
-				Device = DeviceMap.Find(DeviceId);
-			
-				if (!Device || Device->bAuthorizedDevice)
-				{
-					NewDeviceInfo.DeviceId = MoveTemp(DeviceId);
-					NewDeviceInfo.AvdName = MoveTemp(AvdName);
-					NewDeviceInfo.bAuthorizedDevice = false;
-					//note: AndroidTargetDevice::GetName() does not fetch this value, do not rely on this
-					NewDeviceInfo.DeviceName = TEXT("Unauthorized - enable USB debugging");
-					
-					NewDeviceInfo.SerialNumber = MoveTemp(SerialNumber);
+			FAndroidDeviceInfo* Device = DeviceMap.Find(DeviceId);
 			
-					// add the device to the map
-					{
-						FScopeLock ScopeLock(DeviceMapLock);
-
-						if (!Device)
-						{
-							Device = &DeviceMap.Add(NewDeviceInfo.DeviceId);
-						}
-
-						*Device = MoveTemp(NewDeviceInfo);
-					}
-				}
-			}
-			else
+			if (!Device || Device->bAuthorizedDevice != bAuthorized || Device->SerialNumber != SerialNumber)
 			{
-				Device = nullptr;
-				
-				NewDeviceInfo.DeviceId = MoveTemp(DeviceId);
-				NewDeviceInfo.AvdName = MoveTemp(AvdName);
-				NewDeviceInfo.bAuthorizedDevice = true;
+				FAndroidDeviceInfo NewDeviceInfo;
 
-				// grab the Android version
-				const FString AndroidVersionCommand = FString::Printf(TEXT("-s %s %s ro.build.version.release"), *SerialNumber, *GetPropCommand);
-				if (!ExecuteAdbCommand(*AndroidVersionCommand, &NewDeviceInfo.HumanAndroidVersion, nullptr))
+				if (!bAuthorized)
 				{
-					continue;
-				}
-				NewDeviceInfo.HumanAndroidVersion = NewDeviceInfo.HumanAndroidVersion.Replace(TEXT("\r"), TEXT("")).Replace(TEXT("\n"), TEXT(""));
-				NewDeviceInfo.HumanAndroidVersion.TrimStartAndEndInline();
-
-				// grab the Android SDK version
-				const FString SDKVersionCommand = FString::Printf(TEXT("-s %s %s ro.build.version.sdk"), *SerialNumber, *GetPropCommand);
-				FString SDKVersionString;
-				if (!ExecuteAdbCommand(*SDKVersionCommand, &SDKVersionString, nullptr))
-				{
-					continue;
-				}
-				NewDeviceInfo.SDKVersion = FCString::Atoi(*SDKVersionString);
-				if (NewDeviceInfo.SDKVersion <= 0)
-				{
-					NewDeviceInfo.SDKVersion = INDEX_NONE;
-				}
-
-				if (FString AbiOutput; !ExecuteAdbCommand(*FString::Printf(TEXT("-s %s shell getprop ro.product.cpu.abi"), *SerialNumber), &AbiOutput, nullptr))
-				{
-					continue;
+					//note: AndroidTargetDevice::GetName() does not fetch this value, do not rely on this
+					NewDeviceInfo.DeviceName = TEXT("Unauthorized - enable USB debugging");
 				}
 				else
 				{
-					NewDeviceInfo.Architecture = AbiToArchitecture(FStringView{AbiOutput}.TrimStartAndEnd());
-				}
-
-				if (bGetExtensionsViaSurfaceFlinger)
-				{
-					// get the GL extensions string (and a bunch of other stuff)
-					const FString ExtensionsCommand = FString::Printf(TEXT("-s %s shell dumpsys SurfaceFlinger"), *SerialNumber);
-					if (!ExecuteAdbCommand(*ExtensionsCommand, &NewDeviceInfo.GLESExtensions, nullptr))
+					// grab the Android version
+					const FString AndroidVersionCommand = FString::Printf(TEXT("-s %s %s ro.build.version.release"), *SerialNumber, *GetPropCommand);
+					if (!ExecuteAdbCommand(*AndroidVersionCommand, &NewDeviceInfo.HumanAndroidVersion, nullptr))
 					{
 						continue;
 					}
+					NewDeviceInfo.HumanAndroidVersion = NewDeviceInfo.HumanAndroidVersion.Replace(TEXT("\r"), TEXT("")).Replace(TEXT("\n"), TEXT(""));
+					NewDeviceInfo.HumanAndroidVersion.TrimStartAndEndInline();
 
-					// extract DPI information
-					int32 XDpi = ExtractDPI(NewDeviceInfo.GLESExtensions, TEXT("x-dpi"));
-					int32 YDpi = ExtractDPI(NewDeviceInfo.GLESExtensions, TEXT("y-dpi"));
+					// grab the Android SDK version
+					const FString SDKVersionCommand = FString::Printf(TEXT("-s %s %s ro.build.version.sdk"), *SerialNumber, *GetPropCommand);
+					FString SDKVersionString;
+					if (!ExecuteAdbCommand(*SDKVersionCommand, &SDKVersionString, nullptr))
+					{
+						continue;
+					}
+					NewDeviceInfo.SDKVersion = FCString::Atoi(*SDKVersionString);
+					if (NewDeviceInfo.SDKVersion <= 0)
+					{
+						NewDeviceInfo.SDKVersion = INDEX_NONE;
+					}
 
-					if (XDpi != INDEX_NONE && YDpi != INDEX_NONE)
+					if (FString AbiOutput; !ExecuteAdbCommand(*FString::Printf(TEXT("-s %s shell getprop ro.product.cpu.abi"), *SerialNumber), &AbiOutput, nullptr))
+					{
+						continue;
+					}
+					else
 					{
-						NewDeviceInfo.DeviceDPI = (XDpi + YDpi) / 2;
+						NewDeviceInfo.Architecture = AbiToArchitecture(FStringView{AbiOutput}.TrimStartAndEnd());
 					}
 
-					// extract OpenGL version and GPU family name
-					ExtractGPUInfo(NewDeviceInfo.OpenGLVersionString, NewDeviceInfo.GPUFamilyString, NewDeviceInfo.GLESExtensions);
-				}
+					if (bGetExtensionsViaSurfaceFlinger)
+					{
+						// get the GL extensions string (and a bunch of other stuff)
+						const FString ExtensionsCommand = FString::Printf(TEXT("-s %s shell dumpsys SurfaceFlinger"), *SerialNumber);
+						if (!ExecuteAdbCommand(*ExtensionsCommand, &NewDeviceInfo.GLESExtensions, nullptr))
+						{
+							continue;
+						}
 
-				// grab device brand
-				{
-					FString ExecCommand = FString::Printf(TEXT("-s %s %s ro.product.brand"), *SerialNumber, *GetPropCommand);
+						// extract DPI information
+						int32 XDpi = ExtractDPI(NewDeviceInfo.GLESExtensions, TEXT("x-dpi"));
+						int32 YDpi = ExtractDPI(NewDeviceInfo.GLESExtensions, TEXT("y-dpi"));
 
-					FString RoProductBrand;
-					ExecuteAdbCommand(*ExecCommand, &RoProductBrand, nullptr);
-					const TCHAR* Ptr = *RoProductBrand;
-					FParse::Line(&Ptr, NewDeviceInfo.DeviceBrand);
-				}
+						if (XDpi != INDEX_NONE && YDpi != INDEX_NONE)
+						{
+							NewDeviceInfo.DeviceDPI = (XDpi + YDpi) / 2;
+						}
 
-				// grab screen resolution
-				{
-					FString ResolutionString;
-					const FString ExecCommand = FString::Printf(TEXT("-s %s shell wm size"), *SerialNumber);
-					if (ExecuteAdbCommand(*ExecCommand, &ResolutionString, nullptr))
+						// extract OpenGL version and GPU family name
+						ExtractGPUInfo(NewDeviceInfo.OpenGLVersionString, NewDeviceInfo.GPUFamilyString, NewDeviceInfo.GLESExtensions);
+					}
+
+					// grab device brand
 					{
-						bool bFoundResX = false;
-						int32 StartIndex = INDEX_NONE;
-						for (int32 Index = 0; Index < ResolutionString.Len(); ++Index)
+						FString ExecCommand = FString::Printf(TEXT("-s %s %s ro.product.brand"), *SerialNumber, *GetPropCommand);
+
+						FString RoProductBrand;
+						ExecuteAdbCommand(*ExecCommand, &RoProductBrand, nullptr);
+						const TCHAR* Ptr = *RoProductBrand;
+						FParse::Line(&Ptr, NewDeviceInfo.DeviceBrand);
+					}
+
+					// grab screen resolution
+					{
+						FString ResolutionString;
+						const FString ExecCommand = FString::Printf(TEXT("-s %s shell wm size"), *SerialNumber);
+						if (ExecuteAdbCommand(*ExecCommand, &ResolutionString, nullptr))
 						{
-							if (StartIndex == INDEX_NONE && FChar::IsDigit(ResolutionString[Index]))
-							{
-								StartIndex = Index;
-							}
-							else if (StartIndex != INDEX_NONE && !FChar::IsDigit(ResolutionString[Index]))
+							bool bFoundResX = false;
+							int32 StartIndex = INDEX_NONE;
+							for (int32 Index = 0; Index < ResolutionString.Len(); ++Index)
 							{
-								FString str = ResolutionString.Mid(StartIndex, Index - StartIndex);
-
-								if (bFoundResX)
+								if (StartIndex == INDEX_NONE && FChar::IsDigit(ResolutionString[Index]))
 								{
-									NewDeviceInfo.ResolutionY = FCString::Atoi(*str);
-									break;
+									StartIndex = Index;
 								}
-								else
+								else if (StartIndex != INDEX_NONE && !FChar::IsDigit(ResolutionString[Index]))
 								{
-									NewDeviceInfo.ResolutionX = FCString::Atoi(*str);
-									bFoundResX = true;
-									StartIndex = INDEX_NONE;
+									FString str = ResolutionString.Mid(StartIndex, Index - StartIndex);
+
+									if (bFoundResX)
+									{
+										NewDeviceInfo.ResolutionY = FCString::Atoi(*str);
+										break;
+									}
+									else
+									{
+										NewDeviceInfo.ResolutionX = FCString::Atoi(*str);
+										bFoundResX = true;
+										StartIndex = INDEX_NONE;
+									}
 								}
 							}
 						}
 					}
-				}
 
-				// grab the GL ES version
-				FString GLESVersionString;
-				const FString GLVersionCommand = FString::Printf(TEXT("-s %s %s ro.opengles.version"), *SerialNumber, *GetPropCommand);
-				if (!ExecuteAdbCommand(*GLVersionCommand, &GLESVersionString, nullptr))
-				{
-					continue;
-				}
-				NewDeviceInfo.GLESVersion = FCString::Atoi(*GLESVersionString);
-
-				// Find the device model
-				FParse::Value(*DeviceString, TEXT("model:"), NewDeviceInfo.Model);
-				// find the product model (this must match java's android.os.build.model)
-				FString ModelCommand = FString::Printf(TEXT("-s %s %s ro.product.model"), *SerialNumber, *GetPropCommand);
-				FString RoProductModel;
-				if (ExecuteAdbCommand(*ModelCommand, &RoProductModel, nullptr) )
-				{
-					if(!RoProductModel.IsEmpty())
+					// grab the GL ES version
+					FString GLESVersionString;
+					const FString GLVersionCommand = FString::Printf(TEXT("-s %s %s ro.opengles.version"), *SerialNumber, *GetPropCommand);
+					if (!ExecuteAdbCommand(*GLVersionCommand, &GLESVersionString, nullptr))
 					{
-						NewDeviceInfo.Model = RoProductModel.TrimStartAndEnd();
+						continue;
 					}
-				}
-				
-				// Find the build ID
-				FString BuildNumberString;
-				const FString BuildNumberCommand = FString::Printf(TEXT("-s %s %s ro.build.display.id"), *SerialNumber, *GetPropCommand);
-				if (ExecuteAdbCommand(*BuildNumberCommand, &BuildNumberString, nullptr))
-				{
-					NewDeviceInfo.BuildNumber = BuildNumberString.TrimStartAndEnd();
-				}
-
-				// Scan lines looking for ContainsTerm
-				auto FindLineContaining = [](const FString& SourceString, const FString& ContainsTerm)
-				{
-					FString result;
-					UE::String::ParseLines(SourceString,
-						[&result, &ContainsTerm](const FStringView& Line)
+					NewDeviceInfo.GLESVersion = FCString::Atoi(*GLESVersionString);
+
+					// Find the device model
+					FParse::Value(*DeviceString, TEXT("model:"), NewDeviceInfo.Model);
+					// find the product model (this must match java's android.os.build.model)
+					FString ModelCommand = FString::Printf(TEXT("-s %s %s ro.product.model"), *SerialNumber, *GetPropCommand);
+					FString RoProductModel;
+					if (ExecuteAdbCommand(*ModelCommand, &RoProductModel, nullptr) )
+					{
+						if(!RoProductModel.IsEmpty())
 						{
-							if (result.IsEmpty() && Line.Contains(ContainsTerm))
-							{
-								result = Line;
-							}
-						});
-
-					return result;
-				};
-
-				// Parse vulkan version:
-				auto MajorVK = [](uint32 Version) { return (((uint32_t)(Version) >> 22) & 0x7FU); };
-				auto MinorVK = [](uint32 Version) { return (((uint32_t)(Version) >> 12) & 0x3FFU); };
-				auto PatchVK = [](uint32 Version) { return ((uint32_t)(Version) & 0xFFFU); };
-				FString FeaturesString;
-				const FString FeaturesStringCommand = FString::Printf(TEXT("-s %s shell pm list features"), *SerialNumber);
-				if (ExecuteAdbCommand(*FeaturesStringCommand, &FeaturesString, nullptr))
-				{
-					FString VulkanVersionLine = FindLineContaining(FeaturesString, TEXT("android.hardware.vulkan.version"));
-					const FRegexPattern RegexPattern(TEXT("android\\.hardware\\.vulkan\\.version=(\\d*)"));
-					FRegexMatcher RegexMatcher(RegexPattern, *VulkanVersionLine);
-					if (RegexMatcher.FindNext())
+							NewDeviceInfo.Model = RoProductModel.TrimStartAndEnd();
+						}
+					}
+					
+					// Find the build ID
+					FString BuildNumberString;
+					const FString BuildNumberCommand = FString::Printf(TEXT("-s %s %s ro.build.display.id"), *SerialNumber, *GetPropCommand);
+					if (ExecuteAdbCommand(*BuildNumberCommand, &BuildNumberString, nullptr))
 					{
-						uint32 PackedVersion = (uint32)FCString::Atoi64(*RegexMatcher.GetCaptureGroup(1));
-						NewDeviceInfo.VulkanVersion = FString::Printf(TEXT("%d.%d.%d"), MajorVK(PackedVersion), MinorVK(PackedVersion), PatchVK(PackedVersion));
+						NewDeviceInfo.BuildNumber = BuildNumberString.TrimStartAndEnd();
 					}
-				}
-
-				// try vkjson:
-				FString VKJsonString;
-				const FString VKJsonStringCommand = FString::Printf(TEXT("-s %s shell cmd gpu vkjson"), *SerialNumber);
-				if (ExecuteAdbCommand(*VKJsonStringCommand, &VKJsonString, nullptr))
-				{
-					FString VulkanVersionLine = FindLineContaining(VKJsonString, TEXT("apiVersion"));
 
-					const FRegexPattern RegexPattern(TEXT("\"apiVersion\"\\s*:\\s*(\\d*)"));
- 					FRegexMatcher RegexMatcher(RegexPattern, *VulkanVersionLine);
- 					if (RegexMatcher.FindNext())
+					// Scan lines looking for ContainsTerm
+					auto FindLineContaining = [](const FString& SourceString, const FString& ContainsTerm)
 					{
-						FString VulkanVersion = RegexMatcher.GetCaptureGroup(1);
-						uint32 PackedVersion = (uint32)FCString::Atoi64(*VulkanVersion);
-						if(PackedVersion>0)
+						FString result;
+						UE::String::ParseLines(SourceString,
+							[&result, &ContainsTerm](const FStringView& Line)
+							{
+								if (result.IsEmpty() && Line.Contains(ContainsTerm))
+								{
+									result = Line;
+								}
+							});
+
+						return result;
+					};
+
+					// Parse vulkan version:
+					auto MajorVK = [](uint32 Version) { return (((uint32_t)(Version) >> 22) & 0x7FU); };
+					auto MinorVK = [](uint32 Version) { return (((uint32_t)(Version) >> 12) & 0x3FFU); };
+					auto PatchVK = [](uint32 Version) { return ((uint32_t)(Version) & 0xFFFU); };
+					FString FeaturesString;
+					const FString FeaturesStringCommand = FString::Printf(TEXT("-s %s shell pm list features"), *SerialNumber);
+					if (ExecuteAdbCommand(*FeaturesStringCommand, &FeaturesString, nullptr))
+					{
+						FString VulkanVersionLine = FindLineContaining(FeaturesString, TEXT("android.hardware.vulkan.version"));
+						const FRegexPattern RegexPattern(TEXT("android\\.hardware\\.vulkan\\.version=(\\d*)"));
+						FRegexMatcher RegexMatcher(RegexPattern, *VulkanVersionLine);
+						if (RegexMatcher.FindNext())
 						{
+							uint32 PackedVersion = (uint32)FCString::Atoi64(*RegexMatcher.GetCaptureGroup(1));
 							NewDeviceInfo.VulkanVersion = FString::Printf(TEXT("%d.%d.%d"), MajorVK(PackedVersion), MinorVK(PackedVersion), PatchVK(PackedVersion));
 						}
 					}
-				}
-				
-				if (NewDeviceInfo.VulkanVersion.IsEmpty())
-				{
-					NewDeviceInfo.VulkanVersion = TEXT("0.0.0");
-				}
 
-				// create the hardware field
-				FString HardwareCommand = FString::Printf(TEXT("-s %s %s ro.hardware"), *SerialNumber, *GetPropCommand);
-				FString RoHardware;
-				{
-					ExecuteAdbCommand(*HardwareCommand, &RoHardware, nullptr);
-					const TCHAR* Ptr = *RoHardware;
-					FParse::Line(&Ptr, NewDeviceInfo.Hardware);
-				}
-				if (RoHardware.Contains(TEXT("qcom")))
-				{
-					HardwareCommand = FString::Printf(TEXT("-s %s %s ro.hardware.chipname"), *SerialNumber, *GetPropCommand);
-					ExecuteAdbCommand(*HardwareCommand, &RoHardware, nullptr);
-					const TCHAR* Ptr = *RoHardware;
-					FParse::Line(&Ptr, NewDeviceInfo.Hardware);
-				}
-				{
-					HardwareCommand = FString::Printf(TEXT("-s %s %s ro.soc.model"), *SerialNumber, *GetPropCommand);
-					FString RoSOCModelIn;
-					FString RoSOCModelOut;
-					ExecuteAdbCommand(*HardwareCommand, &RoSOCModelIn, nullptr);
-					const TCHAR* Ptr = *RoSOCModelIn;
-					FParse::Line(&Ptr, RoSOCModelOut);
-					if (!RoSOCModelOut.IsEmpty())
+					// try vkjson:
+					FString VKJsonString;
+					const FString VKJsonStringCommand = FString::Printf(TEXT("-s %s shell cmd gpu vkjson"), *SerialNumber);
+					if (ExecuteAdbCommand(*VKJsonStringCommand, &VKJsonString, nullptr))
 					{
-						NewDeviceInfo.Hardware = RoSOCModelOut;
-					}
-				}
-
-				// Read hardware from cpuinfo:
-				FString CPUInfoString;
-				const FString CPUInfoCommand = FString::Printf(TEXT("-s %s shell cat /proc/cpuinfo"), *SerialNumber);
-				if (ExecuteAdbCommand(*CPUInfoCommand, &CPUInfoString, nullptr))
-				{
-					FString HardwareLine = FindLineContaining(CPUInfoString, TEXT("Hardware"));
+						FString VulkanVersionLine = FindLineContaining(VKJsonString, TEXT("apiVersion"));
 
-					const FRegexPattern RegexPattern(TEXT("Hardware\\s*:\\s*(.*)"));
-					FRegexMatcher RegexMatcher(RegexPattern, *HardwareLine);
-					if (RegexMatcher.FindNext())
+						const FRegexPattern RegexPattern(TEXT("\"apiVersion\"\\s*:\\s*(\\d*)"));
+ 						FRegexMatcher RegexMatcher(RegexPattern, *VulkanVersionLine);
+ 						if (RegexMatcher.FindNext())
+						{
+							FString VulkanVersion = RegexMatcher.GetCaptureGroup(1);
+							uint32 PackedVersion = (uint32)FCString::Atoi64(*VulkanVersion);
+							if(PackedVersion>0)
+							{
+								NewDeviceInfo.VulkanVersion = FString::Printf(TEXT("%d.%d.%d"), MajorVK(PackedVersion), MinorVK(PackedVersion), PatchVK(PackedVersion));
+							}
+						}
+					}
+					
+					if (NewDeviceInfo.VulkanVersion.IsEmpty())
 					{
-						NewDeviceInfo.Hardware = RegexMatcher.GetCaptureGroup(1);
+						NewDeviceInfo.VulkanVersion = TEXT("0.0.0");
 					}
-				}
-
-				// Total physical mem:
-				FString MemTotalString;
-				const FString MemTotalCommand = FString::Printf(TEXT("-s %s shell cat /proc/meminfo"), *SerialNumber);
-				if (ExecuteAdbCommand(*MemTotalCommand, &MemTotalString, nullptr))
-				{
-					FString MemTotalLine = FindLineContaining(MemTotalString, TEXT("MemTotal"));
 
-					const FRegexPattern RegexPattern(TEXT("MemTotal:\\s*(\\d*)"));
-					FRegexMatcher RegexMatcher(RegexPattern, *MemTotalLine);
-					if (RegexMatcher.FindNext())
+					// create the hardware field
+					FString HardwareCommand = FString::Printf(TEXT("-s %s %s ro.hardware"), *SerialNumber, *GetPropCommand);
+					FString RoHardware;
 					{
-						NewDeviceInfo.TotalPhysicalKB = (uint32)FCString::Atoi64(*RegexMatcher.GetCaptureGroup(1));
+						ExecuteAdbCommand(*HardwareCommand, &RoHardware, nullptr);
+						const TCHAR* Ptr = *RoHardware;
+						FParse::Line(&Ptr, NewDeviceInfo.Hardware);
+					}
+					if (RoHardware.Contains(TEXT("qcom")))
+					{
+						HardwareCommand = FString::Printf(TEXT("-s %s %s ro.hardware.chipname"), *SerialNumber, *GetPropCommand);
+						ExecuteAdbCommand(*HardwareCommand, &RoHardware, nullptr);
+						const TCHAR* Ptr = *RoHardware;
+						FParse::Line(&Ptr, NewDeviceInfo.Hardware);
+					}
+					{
+						HardwareCommand = FString::Printf(TEXT("-s %s %s ro.soc.model"), *SerialNumber, *GetPropCommand);
+						FString RoSOCModelIn;
+						FString RoSOCModelOut;
+						ExecuteAdbCommand(*HardwareCommand, &RoSOCModelIn, nullptr);
+						const TCHAR* Ptr = *RoSOCModelIn;
+						FParse::Line(&Ptr, RoSOCModelOut);
+						if (!RoSOCModelOut.IsEmpty())
+						{
+							NewDeviceInfo.Hardware = RoSOCModelOut;
+						}
 					}
-				}
-
-				// parse the device name
-				FParse::Value(*DeviceString, TEXT("device:"), NewDeviceInfo.DeviceName);
-				if (NewDeviceInfo.DeviceName.IsEmpty())
-				{
-					FString DeviceCommand = FString::Printf(TEXT("-s %s %s ro.product.device"), *SerialNumber, *GetPropCommand);
-					FString RoProductDevice;
-					ExecuteAdbCommand(*DeviceCommand, &RoProductDevice, nullptr);
-					const TCHAR* Ptr = *RoProductDevice;
-					FParse::Line(&Ptr, NewDeviceInfo.DeviceName);
-				}
 
-				// establish port forwarding if we're doing messaging
-				if (TcpMessagingModule != nullptr)
-				{
-					// fill in the port forwarding array if needed
-					if (PortForwardings.Num() == 0)
+					// Read hardware from cpuinfo:
+					FString CPUInfoString;
+					const FString CPUInfoCommand = FString::Printf(TEXT("-s %s shell cat /proc/cpuinfo"), *SerialNumber);
+					if (ExecuteAdbCommand(*CPUInfoCommand, &CPUInfoString, nullptr))
 					{
-						FString ForwardList;
-						if (ExecuteAdbCommand(TEXT("forward --list"), &ForwardList, nullptr))
+						FString HardwareLine = FindLineContaining(CPUInfoString, TEXT("Hardware"));
+
+						const FRegexPattern RegexPattern(TEXT("Hardware\\s*:\\s*(.*)"));
+						FRegexMatcher RegexMatcher(RegexPattern, *HardwareLine);
+						if (RegexMatcher.FindNext())
 						{
-							ForwardList = ForwardList.Replace(TEXT("\r"), TEXT("\n"));
-							ForwardList.ParseIntoArray(PortForwardings, TEXT("\n"), true);
+							NewDeviceInfo.Hardware = RegexMatcher.GetCaptureGroup(1);
 						}
 					}
 
-					// check if this device already has port forwarding enabled for message bus, eg from another editor session
-					for (FString& FwdString : PortForwardings)
+					// Total physical mem:
+					FString MemTotalString;
+					const FString MemTotalCommand = FString::Printf(TEXT("-s %s shell cat /proc/meminfo"), *SerialNumber);
+					if (ExecuteAdbCommand(*MemTotalCommand, &MemTotalString, nullptr))
 					{
-						const TCHAR* Ptr = *FwdString;
-						FString FwdSerialNumber, FwdHostPortString, FwdDevicePortString;
-						uint16 FwdHostPort, FwdDevicePort;
-						if (FParse::Token(Ptr, FwdSerialNumber, false) && FwdSerialNumber == SerialNumber &&
-							FParse::Token(Ptr, FwdHostPortString, false) && FParse::Value(*FwdHostPortString, TEXT("tcp:"), FwdHostPort) &&
-							FParse::Token(Ptr, FwdDevicePortString, false) && FParse::Value(*FwdDevicePortString, TEXT("tcp:"), FwdDevicePort) && FwdDevicePort == 6666)
+						FString MemTotalLine = FindLineContaining(MemTotalString, TEXT("MemTotal"));
+
+						const FRegexPattern RegexPattern(TEXT("MemTotal:\\s*(\\d*)"));
+						FRegexMatcher RegexMatcher(RegexPattern, *MemTotalLine);
+						if (RegexMatcher.FindNext())
 						{
-							NewDeviceInfo.HostMessageBusPort = FwdHostPort;
-							break;
+							NewDeviceInfo.TotalPhysicalKB = (uint32)FCString::Atoi64(*RegexMatcher.GetCaptureGroup(1));
 						}
 					}
 
-					// if not, setup TCP port forwarding for message bus on first available TCP port above 6666
-					if (NewDeviceInfo.HostMessageBusPort == 0)
+					// parse the device name
+					FParse::Value(*DeviceString, TEXT("device:"), NewDeviceInfo.DeviceName);
+					if (NewDeviceInfo.DeviceName.IsEmpty())
 					{
-						uint16 HostMessageBusPort = 6666;
-						bool bFoundPort;
-						do
+						FString DeviceCommand = FString::Printf(TEXT("-s %s %s ro.product.device"), *SerialNumber, *GetPropCommand);
+						FString RoProductDevice;
+						ExecuteAdbCommand(*DeviceCommand, &RoProductDevice, nullptr);
+						const TCHAR* Ptr = *RoProductDevice;
+						FParse::Line(&Ptr, NewDeviceInfo.DeviceName);
+					}
+
+					// establish port forwarding if we're doing messaging
+					if (TcpMessagingModule != nullptr)
+					{
+						// fill in the port forwarding array if needed
+						if (PortForwardings.Num() == 0)
+						{
+							FString ForwardList;
+							if (ExecuteAdbCommand(TEXT("forward --list"), &ForwardList, nullptr))
+							{
+								ForwardList = ForwardList.Replace(TEXT("\r"), TEXT("\n"));
+								ForwardList.ParseIntoArray(PortForwardings, TEXT("\n"), true);
+							}
+						}
+
+						// check if this device already has port forwarding enabled for message bus, eg from another editor session
+						for (FString& FwdString : PortForwardings)
 						{
-							bFoundPort = true;
-							for (auto It = DeviceMap.CreateConstIterator(); It; ++It)
+							const TCHAR* Ptr = *FwdString;
+							FString FwdSerialNumber, FwdHostPortString, FwdDevicePortString;
+							uint16 FwdHostPort, FwdDevicePort;
+							if (FParse::Token(Ptr, FwdSerialNumber, false) && FwdSerialNumber == SerialNumber &&
+								FParse::Token(Ptr, FwdHostPortString, false) && FParse::Value(*FwdHostPortString, TEXT("tcp:"), FwdHostPort) &&
+								FParse::Token(Ptr, FwdDevicePortString, false) && FParse::Value(*FwdDevicePortString, TEXT("tcp:"), FwdDevicePort) && FwdDevicePort == 6666)
 							{
-								if (HostMessageBusPort == It.Value().HostMessageBusPort)
+								NewDeviceInfo.HostMessageBusPort = FwdHostPort;
+								break;
+							}
+						}
+
+						// if not, setup TCP port forwarding for message bus on first available TCP port above 6666
+						if (NewDeviceInfo.HostMessageBusPort == 0)
+						{
+							uint16 HostMessageBusPort = 6666;
+							bool bFoundPort;
+							do
+							{
+								bFoundPort = true;
+								for (auto It = DeviceMap.CreateConstIterator(); It; ++It)
 								{
-									bFoundPort = false;
-									HostMessageBusPort++;
-									break;
+									if (HostMessageBusPort == It.Value().HostMessageBusPort)
+									{
+										bFoundPort = false;
+										HostMessageBusPort++;
+										break;
+									}
 								}
-							}
-						} while (!bFoundPort);
+							} while (!bFoundPort);
 
-						FString DeviceCommand = FString::Printf(TEXT("-s %s forward tcp:%d tcp:6666"), *SerialNumber, HostMessageBusPort);
-						ExecuteAdbCommand(*DeviceCommand, nullptr, nullptr);
-						NewDeviceInfo.HostMessageBusPort = HostMessageBusPort;
-					}
+							FString DeviceCommand = FString::Printf(TEXT("-s %s forward tcp:%d tcp:6666"), *SerialNumber, HostMessageBusPort);
+							ExecuteAdbCommand(*DeviceCommand, nullptr, nullptr);
+							NewDeviceInfo.HostMessageBusPort = HostMessageBusPort;
+						}
 
-					TcpMessagingModule->AddOutgoingConnection(FString::Printf(TEXT("127.0.0.1:%d"), NewDeviceInfo.HostMessageBusPort));
-				}
+						TcpMessagingModule->AddOutgoingConnection(FString::Printf(TEXT("127.0.0.1:%d"), NewDeviceInfo.HostMessageBusPort));
+					}
 
-				// Add reverse port forwarding
-				uint16 ReversePortMappings[]
-				{
-					41899,	// Network file server, DEFAULT_TCP_FILE_SERVING_PORT in NetworkMessage.h
-					1981,	// Unreal Insights data collection, TraceInsightsModule.cpp
+					// Add reverse port forwarding
+					uint16 ReversePortMappings[]
+					{
+						41899,	// Network file server, DEFAULT_TCP_FILE_SERVING_PORT in NetworkMessage.h
+						1981,	// Unreal Insights data collection, TraceInsightsModule.cpp
 #if UE_WITH_ZEN
-					UE::Zen::IsDefaultServicePresent() ? UE::Zen::GetDefaultServiceInstance().GetPort() : uint16(0), // Zen Store, usually defaults to 8558
+						UE::Zen::IsDefaultServicePresent() ? UE::Zen::GetDefaultServiceInstance().GetPort() : uint16(0), // Zen Store, usually defaults to 8558
 #endif
-					0 // end of list
-				};
+						0 // end of list
+					};
 
-				for (int32 Idx=0; ReversePortMappings[Idx] > 0; Idx++)
-				{
-					FString DeviceCommand = FString::Printf(TEXT("-s %s reverse tcp:%d tcp:%d"), *SerialNumber, ReversePortMappings[Idx], ReversePortMappings[Idx]);
-					// It doesn't really matter if a mapping already exists. There is no listening local port so no contention between multiple editor instances
-					ExecuteAdbCommand(*DeviceCommand, nullptr, nullptr);
-				}
+					for (int32 Idx=0; ReversePortMappings[Idx] > 0; Idx++)
+					{
+						FString DeviceCommand = FString::Printf(TEXT("-s %s reverse tcp:%d tcp:%d"), *SerialNumber, ReversePortMappings[Idx], ReversePortMappings[Idx]);
+						// It doesn't really matter if a mapping already exists. There is no listening local port so no contention between multiple editor instances
+						ExecuteAdbCommand(*DeviceCommand, nullptr, nullptr);
+					}
 
-				FString WindowDisplaysOutput;
-				const FString ExtensionsCommand = FString::Printf(TEXT("-s %s shell dumpsys window displays"), *SerialNumber);
-				if (!ExecuteAdbCommand(*ExtensionsCommand, &WindowDisplaysOutput, nullptr))
-				{
-					continue;
-				}
-				else
-				{
-					WindowDisplaysOutput = WindowDisplaysOutput.ToLower();
-					const FRegexPattern RegexPattern(TEXT(".*initcutout.*insets=rect\\((\\d+)\\W*,\\s*(\\d+)\\W*-\\s*(\\d+)\\W*,\\s*(\\d+)\\W*\\)"));
-					FRegexMatcher RegexMatcher(RegexPattern, *WindowDisplaysOutput);
-					if (RegexMatcher.FindNext())
+					FString WindowDisplaysOutput;
+					const FString ExtensionsCommand = FString::Printf(TEXT("-s %s shell dumpsys window displays"), *SerialNumber);
+					if (!ExecuteAdbCommand(*ExtensionsCommand, &WindowDisplaysOutput, nullptr))
 					{
-						// store the insets independently from the resolution.
-						NewDeviceInfo.InsetsLeft = FCString::Atof(*RegexMatcher.GetCaptureGroup(1)) / NewDeviceInfo.ResolutionX;
-						NewDeviceInfo.InsetsTop = FCString::Atof(*RegexMatcher.GetCaptureGroup(2)) / NewDeviceInfo.ResolutionY;
-						NewDeviceInfo.InsetsRight = FCString::Atof(*RegexMatcher.GetCaptureGroup(3)) / NewDeviceInfo.ResolutionX;
-						NewDeviceInfo.InsetsBottom = FCString::Atof(*RegexMatcher.GetCaptureGroup(4)) / NewDeviceInfo.ResolutionY;
+						continue;
+					}
+					else
+					{
+						WindowDisplaysOutput = WindowDisplaysOutput.ToLower();
+						const FRegexPattern RegexPattern(TEXT(".*initcutout.*insets=rect\\((\\d+)\\W*,\\s*(\\d+)\\W*-\\s*(\\d+)\\W*,\\s*(\\d+)\\W*\\)"));
+						FRegexMatcher RegexMatcher(RegexPattern, *WindowDisplaysOutput);
+						if (RegexMatcher.FindNext())
+						{
+							// store the insets independently from the resolution.
+							NewDeviceInfo.InsetsLeft = FCString::Atof(*RegexMatcher.GetCaptureGroup(1)) / NewDeviceInfo.ResolutionX;
+							NewDeviceInfo.InsetsTop = FCString::Atof(*RegexMatcher.GetCaptureGroup(2)) / NewDeviceInfo.ResolutionY;
+							NewDeviceInfo.InsetsRight = FCString::Atof(*RegexMatcher.GetCaptureGroup(3)) / NewDeviceInfo.ResolutionX;
+							NewDeviceInfo.InsetsBottom = FCString::Atof(*RegexMatcher.GetCaptureGroup(4)) / NewDeviceInfo.ResolutionY;
+						}
 					}
 				}
 
+				NewDeviceInfo.DeviceId = MoveTemp(DeviceId);
+				NewDeviceInfo.AvdName = MoveTemp(AvdName);
+				NewDeviceInfo.bAuthorizedDevice = bAuthorized;
 				NewDeviceInfo.SerialNumber = MoveTemp(SerialNumber);
-			
+
+				// add the device to the map
 				{
 					FScopeLock ScopeLock(DeviceMapLock);
 
-					Device = &DeviceMap.Add(NewDeviceInfo.DeviceId);
+					if (!Device)
+					{
+						Device = &DeviceMap.Add(NewDeviceInfo.DeviceId);
+					}
+
 					*Device = MoveTemp(NewDeviceInfo);
 				}
 			}
@@ -786,7 +771,11 @@ private:
 				{
 					FScopeLock ScopeLock(DeviceMapLock);
 
-					Device = &DeviceMap.Add(NewDeviceInfo.DeviceId);
+					if (!Device)
+					{
+						Device = &DeviceMap.Add(NewDeviceInfo.DeviceId);
+					}
+
 					*Device = MoveTemp(NewDeviceInfo);
 				}
 			}
