Added quick thread stopping and adb timeout in AndroidDeviceDetection

#android

diff --git a/Engine/Source/Developer/Android/AndroidDeviceDetection/Private/AndroidDeviceDetectionModule.cpp b/Engine/Source/Developer/Android/AndroidDeviceDetection/Private/AndroidDeviceDetectionModule.cpp
--- a/Engine/Source/Developer/Android/AndroidDeviceDetection/Private/AndroidDeviceDetectionModule.cpp
+++ b/Engine/Source/Developer/Android/AndroidDeviceDetection/Private/AndroidDeviceDetectionModule.cpp
@@ -5,6 +5,7 @@
 =============================================================================*/
 
 #include "CoreTypes.h"
+#include "Async/EventCount.h"
 #include "HAL/UnrealMemory.h"
 #include "Containers/Array.h"
 #include "Containers/UnrealString.h"
@@ -19,7 +20,6 @@
 #include "Misc/Paths.h"
 #include "HAL/Runnable.h"
 #include "HAL/RunnableThread.h"
-#include "HAL/ThreadSafeCounter.h"
 #include "Misc/ScopeLock.h"
 #include "Modules/ModuleManager.h"
 #include "Interfaces/IAndroidDeviceDetection.h"
@@ -61,7 +61,6 @@ class FAndroidDeviceDetectionRunnable : public FRunnable
 {
 public:
 	FAndroidDeviceDetectionRunnable(TMap<FString,FAndroidDeviceInfo>& InDeviceMap, FCriticalSection* InDeviceMapLock, FCriticalSection* InADBPathCheckLock) :
-		StopTaskCounter(0),
 		DeviceMap(InDeviceMap),
 		DeviceMapLock(InDeviceMapLock),
 		ADBPathCheckLock(InADBPathCheckLock),
@@ -83,16 +82,20 @@ public:
 	{
 	}
 
-	virtual void Stop(void)
+	virtual void Stop()
 	{
-		StopTaskCounter.Increment();
+		bStopRequested.store(true, std::memory_order_seq_cst);
+		StopEvent.Notify();
 	}
 
-	virtual uint32 Run(void)
+	virtual uint32 Run()
 	{
-		int LoopCount = 10;
+		if (bStopRequested.load(std::memory_order_relaxed))
+		{
+			return 0;
+		}
 
-		while (StopTaskCounter.GetValue() == 0)
+		for (int32 LoopCount = 10;;)
 		{
 			// query when we have waited 'GAndroidDeviceDetectionPollInterval' seconds.
 			if (LoopCount++ >= GAndroidDeviceDetectionPollInterval || ForceCheck)
@@ -100,13 +103,21 @@ public:
 				// Make sure we have an ADB path before checking
 				FScopeLock PathLock(ADBPathCheckLock);
 				if (HasADBPath)
+				{
 					QueryConnectedDevices();
+				}
 
 				LoopCount = 0;
 				ForceCheck = false;
 			}
 
-			FPlatformProcess::Sleep(1.0f);
+			{
+				UE::FEventCountToken Token = StopEvent.PrepareWait();
+				if (bStopRequested.load(std::memory_order_relaxed) || StopEvent.WaitFor(Token, UE::FMonotonicTimeSpan::FromSeconds(1.0)))
+				{
+					break;
+				}
+			}
 		}
 
 		return 0;
@@ -131,30 +142,61 @@ public:
 	}
 
 private:
-
-	bool ExecuteAdbCommand( const FString& CommandLine, FString* OutStdOut, FString* OutStdErr ) const
+	bool ExecuteAdbCommand(const FString& CommandLine, FString* OutStdOut) const
 	{
-		// execute the command
-		int32 ReturnCode;
-		FString DefaultError;
-
-		// make sure there's a place for error output to go if the caller specified nullptr
-		if (!OutStdErr)
+		if (bStopRequested.load(std::memory_order_relaxed) || !FPaths::FileExists(ADBPath))
 		{
-			OutStdErr = &DefaultError;
+			return false;
 		}
 
-		if (FPaths::FileExists(ADBPath))
+		void* ReadPipe;
+		void* WritePipe;
+		if (!FPlatformProcess::CreatePipe(ReadPipe, WritePipe))
 		{
-			FPlatformProcess::ExecProcess(*ADBPath, *CommandLine, &ReturnCode, OutStdOut, OutStdErr);
+			return false;
+		}
 
-			if (ReturnCode != 0)
+		int32 ReturnCode = -1;
+		
+		{
+			FProcHandle Process = FPlatformProcess::CreateProc(*ADBPath, *CommandLine, true, true, false, nullptr, 0, nullptr, WritePipe, ReadPipe);
+			if (Process.IsValid())
 			{
-				FPlatformMisc::LowLevelOutputDebugStringf(TEXT("The Android SDK command '%s' failed to run. Return code: %d, Error: %s\n"), *CommandLine, ReturnCode, **OutStdErr);
+				for (UE::FMonotonicTimePoint MaxTime = UE::FMonotonicTimePoint::Now() + UE::FMonotonicTimeSpan::FromSeconds(10.0);;)
+				{
+					bool bIsRunning = FPlatformProcess::IsProcRunning(Process);
+					
+					if (OutStdOut)
+					{
+						*OutStdOut += FPlatformProcess::ReadPipe(ReadPipe);
+					}
 
-				return false;
+					if (!bIsRunning)
+					{
+						verify(FPlatformProcess::GetProcReturnCode(Process, &ReturnCode));
+						break;
+					}
+
+					FPlatformProcess::Sleep(0.0);
+
+					if (bStopRequested.load(std::memory_order_relaxed) || UE::FMonotonicTimePoint::Now() >= MaxTime)
+					{
+						FPlatformProcess::TerminateProc(Process, true);
+						break;
+					}
+				}
+
+				FPlatformProcess::CloseProc(Process);
 			}
 		}
+		
+		FPlatformProcess::ClosePipe(ReadPipe, WritePipe);
+		
+		if (ReturnCode != 0)
+		{
+			FPlatformMisc::LowLevelOutputDebugStringf(TEXT("The Android SDK command '%s' failed to run. Return code: %d\n"), *CommandLine, ReturnCode);
+			return false;
+		}
 
 		return true;
 	}
@@ -256,7 +298,7 @@ private:
 	{
 		// grab the list of devices via adb
 		FString StdOut;
-		if (!ExecuteAdbCommand(TEXT("devices -l"), &StdOut, nullptr))
+		if (!ExecuteAdbCommand(TEXT("devices -l"), &StdOut))
 		{
 			return;
 		}
@@ -306,6 +348,11 @@ private:
 
 		for (const FString& DeviceString : DeviceStrings)
 		{
+			if (bStopRequested.load(std::memory_order_relaxed))
+			{
+				return;
+			}
+			
 			// skip over non-device lines
 			if (DeviceString.StartsWith("* ") || DeviceString.StartsWith("List "))
 			{
@@ -336,7 +383,7 @@ private:
 			FString SerialNumber = DeviceString.Left(TabIndex), DeviceId, AvdName;
 
 			// Find the AVD name
-			if (FString AvdNameOutput; SerialNumber.StartsWith("emulator-") && ExecuteAdbCommand(*(TEXT("-s ") + SerialNumber + TEXT(" emu avd name")), &AvdNameOutput, nullptr))
+			if (FString AvdNameOutput; SerialNumber.StartsWith("emulator-") && ExecuteAdbCommand(*(TEXT("-s ") + SerialNumber + TEXT(" emu avd name")), &AvdNameOutput))
 			{
 				const TCHAR* Ptr = *AvdNameOutput;
 				FParse::Line(&Ptr, AvdName);
@@ -364,7 +411,7 @@ private:
 				{
 					// grab the Android version
 					const FString AndroidVersionCommand = FString::Printf(TEXT("-s %s %s ro.build.version.release"), *SerialNumber, *GetPropCommand);
-					if (!ExecuteAdbCommand(*AndroidVersionCommand, &NewDeviceInfo.HumanAndroidVersion, nullptr))
+					if (!ExecuteAdbCommand(*AndroidVersionCommand, &NewDeviceInfo.HumanAndroidVersion))
 					{
 						continue;
 					}
@@ -374,7 +421,7 @@ private:
 					// grab the Android SDK version
 					const FString SDKVersionCommand = FString::Printf(TEXT("-s %s %s ro.build.version.sdk"), *SerialNumber, *GetPropCommand);
 					FString SDKVersionString;
-					if (!ExecuteAdbCommand(*SDKVersionCommand, &SDKVersionString, nullptr))
+					if (!ExecuteAdbCommand(*SDKVersionCommand, &SDKVersionString))
 					{
 						continue;
 					}
@@ -384,7 +431,7 @@ private:
 						NewDeviceInfo.SDKVersion = INDEX_NONE;
 					}
 
-					if (FString AbiOutput; !ExecuteAdbCommand(*FString::Printf(TEXT("-s %s shell getprop ro.product.cpu.abi"), *SerialNumber), &AbiOutput, nullptr))
+					if (FString AbiOutput; !ExecuteAdbCommand(*FString::Printf(TEXT("-s %s shell getprop ro.product.cpu.abi"), *SerialNumber), &AbiOutput))
 					{
 						continue;
 					}
@@ -397,7 +444,7 @@ private:
 					{
 						// get the GL extensions string (and a bunch of other stuff)
 						const FString ExtensionsCommand = FString::Printf(TEXT("-s %s shell dumpsys SurfaceFlinger"), *SerialNumber);
-						if (!ExecuteAdbCommand(*ExtensionsCommand, &NewDeviceInfo.GLESExtensions, nullptr))
+						if (!ExecuteAdbCommand(*ExtensionsCommand, &NewDeviceInfo.GLESExtensions))
 						{
 							continue;
 						}
@@ -420,7 +467,7 @@ private:
 						FString ExecCommand = FString::Printf(TEXT("-s %s %s ro.product.brand"), *SerialNumber, *GetPropCommand);
 
 						FString RoProductBrand;
-						ExecuteAdbCommand(*ExecCommand, &RoProductBrand, nullptr);
+						ExecuteAdbCommand(*ExecCommand, &RoProductBrand);
 						const TCHAR* Ptr = *RoProductBrand;
 						FParse::Line(&Ptr, NewDeviceInfo.DeviceBrand);
 					}
@@ -429,7 +476,7 @@ private:
 					{
 						FString ResolutionString;
 						const FString ExecCommand = FString::Printf(TEXT("-s %s shell wm size"), *SerialNumber);
-						if (ExecuteAdbCommand(*ExecCommand, &ResolutionString, nullptr))
+						if (ExecuteAdbCommand(*ExecCommand, &ResolutionString))
 						{
 							bool bFoundResX = false;
 							int32 StartIndex = INDEX_NONE;
@@ -462,7 +509,7 @@ private:
 					// grab the GL ES version
 					FString GLESVersionString;
 					const FString GLVersionCommand = FString::Printf(TEXT("-s %s %s ro.opengles.version"), *SerialNumber, *GetPropCommand);
-					if (!ExecuteAdbCommand(*GLVersionCommand, &GLESVersionString, nullptr))
+					if (!ExecuteAdbCommand(*GLVersionCommand, &GLESVersionString))
 					{
 						continue;
 					}
@@ -473,7 +520,7 @@ private:
 					// find the product model (this must match java's android.os.build.model)
 					FString ModelCommand = FString::Printf(TEXT("-s %s %s ro.product.model"), *SerialNumber, *GetPropCommand);
 					FString RoProductModel;
-					if (ExecuteAdbCommand(*ModelCommand, &RoProductModel, nullptr) )
+					if (ExecuteAdbCommand(*ModelCommand, &RoProductModel) )
 					{
 						if(!RoProductModel.IsEmpty())
 						{
@@ -484,7 +531,7 @@ private:
 					// Find the build ID
 					FString BuildNumberString;
 					const FString BuildNumberCommand = FString::Printf(TEXT("-s %s %s ro.build.display.id"), *SerialNumber, *GetPropCommand);
-					if (ExecuteAdbCommand(*BuildNumberCommand, &BuildNumberString, nullptr))
+					if (ExecuteAdbCommand(*BuildNumberCommand, &BuildNumberString))
 					{
 						NewDeviceInfo.BuildNumber = BuildNumberString.TrimStartAndEnd();
 					}
@@ -511,7 +558,7 @@ private:
 					auto PatchVK = [](uint32 Version) { return ((uint32_t)(Version) & 0xFFFU); };
 					FString FeaturesString;
 					const FString FeaturesStringCommand = FString::Printf(TEXT("-s %s shell pm list features"), *SerialNumber);
-					if (ExecuteAdbCommand(*FeaturesStringCommand, &FeaturesString, nullptr))
+					if (ExecuteAdbCommand(*FeaturesStringCommand, &FeaturesString))
 					{
 						FString VulkanVersionLine = FindLineContaining(FeaturesString, TEXT("android.hardware.vulkan.version"));
 						const FRegexPattern RegexPattern(TEXT("android\\.hardware\\.vulkan\\.version=(\\d*)"));
@@ -526,7 +573,7 @@ private:
 					// try vkjson:
 					FString VKJsonString;
 					const FString VKJsonStringCommand = FString::Printf(TEXT("-s %s shell cmd gpu vkjson"), *SerialNumber);
-					if (ExecuteAdbCommand(*VKJsonStringCommand, &VKJsonString, nullptr))
+					if (ExecuteAdbCommand(*VKJsonStringCommand, &VKJsonString))
 					{
 						FString VulkanVersionLine = FindLineContaining(VKJsonString, TEXT("apiVersion"));
 
@@ -552,14 +599,14 @@ private:
 					FString HardwareCommand = FString::Printf(TEXT("-s %s %s ro.hardware"), *SerialNumber, *GetPropCommand);
 					FString RoHardware;
 					{
-						ExecuteAdbCommand(*HardwareCommand, &RoHardware, nullptr);
+						ExecuteAdbCommand(*HardwareCommand, &RoHardware);
 						const TCHAR* Ptr = *RoHardware;
 						FParse::Line(&Ptr, NewDeviceInfo.Hardware);
 					}
 					if (RoHardware.Contains(TEXT("qcom")))
 					{
 						HardwareCommand = FString::Printf(TEXT("-s %s %s ro.hardware.chipname"), *SerialNumber, *GetPropCommand);
-						ExecuteAdbCommand(*HardwareCommand, &RoHardware, nullptr);
+						ExecuteAdbCommand(*HardwareCommand, &RoHardware);
 						const TCHAR* Ptr = *RoHardware;
 						FParse::Line(&Ptr, NewDeviceInfo.Hardware);
 					}
@@ -567,7 +614,7 @@ private:
 						HardwareCommand = FString::Printf(TEXT("-s %s %s ro.soc.model"), *SerialNumber, *GetPropCommand);
 						FString RoSOCModelIn;
 						FString RoSOCModelOut;
-						ExecuteAdbCommand(*HardwareCommand, &RoSOCModelIn, nullptr);
+						ExecuteAdbCommand(*HardwareCommand, &RoSOCModelIn);
 						const TCHAR* Ptr = *RoSOCModelIn;
 						FParse::Line(&Ptr, RoSOCModelOut);
 						if (!RoSOCModelOut.IsEmpty())
@@ -579,7 +626,7 @@ private:
 					// Read hardware from cpuinfo:
 					FString CPUInfoString;
 					const FString CPUInfoCommand = FString::Printf(TEXT("-s %s shell cat /proc/cpuinfo"), *SerialNumber);
-					if (ExecuteAdbCommand(*CPUInfoCommand, &CPUInfoString, nullptr))
+					if (ExecuteAdbCommand(*CPUInfoCommand, &CPUInfoString))
 					{
 						FString HardwareLine = FindLineContaining(CPUInfoString, TEXT("Hardware"));
 
@@ -594,7 +641,7 @@ private:
 					// Total physical mem:
 					FString MemTotalString;
 					const FString MemTotalCommand = FString::Printf(TEXT("-s %s shell cat /proc/meminfo"), *SerialNumber);
-					if (ExecuteAdbCommand(*MemTotalCommand, &MemTotalString, nullptr))
+					if (ExecuteAdbCommand(*MemTotalCommand, &MemTotalString))
 					{
 						FString MemTotalLine = FindLineContaining(MemTotalString, TEXT("MemTotal"));
 
@@ -612,7 +659,7 @@ private:
 					{
 						FString DeviceCommand = FString::Printf(TEXT("-s %s %s ro.product.device"), *SerialNumber, *GetPropCommand);
 						FString RoProductDevice;
-						ExecuteAdbCommand(*DeviceCommand, &RoProductDevice, nullptr);
+						ExecuteAdbCommand(*DeviceCommand, &RoProductDevice);
 						const TCHAR* Ptr = *RoProductDevice;
 						FParse::Line(&Ptr, NewDeviceInfo.DeviceName);
 					}
@@ -624,7 +671,7 @@ private:
 						if (PortForwardings.Num() == 0)
 						{
 							FString ForwardList;
-							if (ExecuteAdbCommand(TEXT("forward --list"), &ForwardList, nullptr))
+							if (ExecuteAdbCommand(TEXT("forward --list"), &ForwardList))
 							{
 								ForwardList = ForwardList.Replace(TEXT("\r"), TEXT("\n"));
 								ForwardList.ParseIntoArray(PortForwardings, TEXT("\n"), true);
@@ -666,7 +713,7 @@ private:
 							} while (!bFoundPort);
 
 							FString DeviceCommand = FString::Printf(TEXT("-s %s forward tcp:%d tcp:6666"), *SerialNumber, HostMessageBusPort);
-							ExecuteAdbCommand(*DeviceCommand, nullptr, nullptr);
+							ExecuteAdbCommand(*DeviceCommand, nullptr);
 							NewDeviceInfo.HostMessageBusPort = HostMessageBusPort;
 						}
 
@@ -688,12 +735,12 @@ private:
 					{
 						FString DeviceCommand = FString::Printf(TEXT("-s %s reverse tcp:%d tcp:%d"), *SerialNumber, ReversePortMappings[Idx], ReversePortMappings[Idx]);
 						// It doesn't really matter if a mapping already exists. There is no listening local port so no contention between multiple editor instances
-						ExecuteAdbCommand(*DeviceCommand, nullptr, nullptr);
+						ExecuteAdbCommand(*DeviceCommand, nullptr);
 					}
 
 					FString WindowDisplaysOutput;
 					const FString ExtensionsCommand = FString::Printf(TEXT("-s %s shell dumpsys window displays"), *SerialNumber);
-					if (!ExecuteAdbCommand(*ExtensionsCommand, &WindowDisplaysOutput, nullptr))
+					if (!ExecuteAdbCommand(*ExtensionsCommand, &WindowDisplaysOutput))
 					{
 						continue;
 					}
@@ -736,6 +783,11 @@ private:
 
 		for (FString& AvdName : AvdNames)
 		{
+			if (bStopRequested.load(std::memory_order_relaxed))
+			{
+				return;
+			}
+
 			FString DeviceId = TEXT("avd-") + AvdName;
 			FAndroidDeviceInfo* Device = DeviceMap.Find(DeviceId);
 			
@@ -810,6 +862,8 @@ private:
 	}
 
 private:
+	mutable UE::FEventCount StopEvent;
+	std::atomic_bool bStopRequested = false;
 
 	// path to the adb command
 	FString ADBPath;
@@ -817,9 +871,6 @@ private:
 	FString GetPropCommand;
 	bool bGetExtensionsViaSurfaceFlinger;
 
-	// > 0 if we've been asked to abort work in progress at the next opportunity
-	FThreadSafeCounter StopTaskCounter;
-
 	TMap<FString,FAndroidDeviceInfo>& DeviceMap;
 	FCriticalSection* DeviceMapLock;
 
