Fix to UAT hanging/timingout when adb is called in the editor

#tests Tested failed UAT job now succeeds. Tested Editor still detects Android devices, and shuts down properly

diff --git a/Engine/Source/Developer/Android/AndroidDeviceDetection/Private/AndroidDeviceDetectionModule.cpp b/Engine/Source/Developer/Android/AndroidDeviceDetection/Private/AndroidDeviceDetectionModule.cpp
--- a/Engine/Source/Developer/Android/AndroidDeviceDetection/Private/AndroidDeviceDetectionModule.cpp
+++ b/Engine/Source/Developer/Android/AndroidDeviceDetection/Private/AndroidDeviceDetectionModule.cpp
@@ -44,6 +44,152 @@
 #endif
 #include "String/ParseLines.h"
 
+// FIX in 5.7 by removing this and calling FPlatformProcess::CreateProc directly once this code is moved there
+#if PLATFORM_WINDOWS
+#include "Windows/AllowWindowsPlatformTypes.h"
+static FProcHandle CreateProc(const TCHAR* URL, const TCHAR* Parms, bool bLaunchDetached, bool bLaunchHidden, bool bLaunchReallyHidden, uint32* OutProcessID, int32 PriorityModifier, const TCHAR* OptionalWorkingDirectory, void* PipeWriteChild)
+{
+	auto HandleError = [&URL, &Parms, &OutProcessID]
+	{
+		DWORD ErrorCode = GetLastError();
+
+		TCHAR ErrorMessage[512];
+		FWindowsPlatformMisc::GetSystemErrorMessage(ErrorMessage, 512, ErrorCode);
+
+		UE_LOG(LogWindows, Warning, TEXT("CreateProc failed: %s (0x%08x)"), ErrorMessage, ErrorCode);
+		if (ErrorCode == ERROR_NOT_ENOUGH_MEMORY || ErrorCode == ERROR_OUTOFMEMORY)
+		{
+			// These errors are common enough that we want some available memory information
+			FPlatformMemoryStats Stats = FPlatformMemory::GetStats();
+			UE_LOG(LogWindows, Warning, TEXT("Mem used: %.2f MB, OS Free %.2f MB"), Stats.UsedPhysical / 1048576.0f, Stats.AvailablePhysical / 1048576.0f);
+		}
+		UE_LOG(LogWindows, Warning, TEXT("URL: %s %s"), URL, Parms);
+		if (OutProcessID != nullptr)
+		{
+			*OutProcessID = 0;
+		}
+	};
+	
+	//UE_LOG(LogWindows, Log,  TEXT("CreateProc %s %s"), URL, Parms );
+
+	// initialize process creation flags
+	uint32 CreateFlags = NORMAL_PRIORITY_CLASS;
+	if (PriorityModifier < 0)
+	{
+		CreateFlags = (PriorityModifier == -1) ? BELOW_NORMAL_PRIORITY_CLASS : IDLE_PRIORITY_CLASS;
+	}
+	else if (PriorityModifier > 0)
+	{
+		CreateFlags = (PriorityModifier == 1) ? ABOVE_NORMAL_PRIORITY_CLASS : HIGH_PRIORITY_CLASS;
+	}
+
+	if (bLaunchDetached)
+	{
+		CreateFlags |= DETACHED_PROCESS;
+	}
+
+	// initialize window flags
+	uint32 dwFlags = 0;
+	uint16 ShowWindowFlags = SW_HIDE;
+	if (bLaunchReallyHidden)
+	{
+		dwFlags = STARTF_USESHOWWINDOW;
+	}
+	else if (bLaunchHidden)
+	{
+		dwFlags = STARTF_USESHOWWINDOW;
+		ShowWindowFlags = SW_SHOWMINNOACTIVE;
+	}
+
+	TArray<uint8> AttributeList;
+	{
+		HANDLE Handles[1];
+		uint32 HandleCount = 0;
+		
+		if (PipeWriteChild)
+		{
+			Handles[HandleCount++] = HANDLE(PipeWriteChild);
+		}
+		
+		if (HandleCount)
+		{
+			SIZE_T BufferSize = 0;
+			verify(!InitializeProcThreadAttributeList(nullptr, 1, 0, &BufferSize));
+			check(GetLastError() == ERROR_INSUFFICIENT_BUFFER);
+			AttributeList.SetNum(BufferSize);
+
+			if (!InitializeProcThreadAttributeList(LPPROC_THREAD_ATTRIBUTE_LIST(AttributeList.GetData()), 1, 0, &BufferSize))
+			{
+				HandleError();
+				return {};
+			}
+			if (!UpdateProcThreadAttribute(LPPROC_THREAD_ATTRIBUTE_LIST(AttributeList.GetData()), 0, PROC_THREAD_ATTRIBUTE_HANDLE_LIST, Handles, HandleCount * sizeof(HANDLE), nullptr, nullptr))
+			{
+				HandleError();
+				DeleteProcThreadAttributeList(LPPROC_THREAD_ATTRIBUTE_LIST(AttributeList.GetData()));
+				return {};
+			}
+
+			CreateFlags |= EXTENDED_STARTUPINFO_PRESENT;
+			dwFlags |= STARTF_USESTDHANDLES;
+		}
+	}
+
+	// initialize startup info
+	STARTUPINFOEX StartupInfoEx
+	{
+		sizeof(STARTUPINFOEX),
+		NULL, NULL, NULL,
+		(DWORD)CW_USEDEFAULT,
+		(DWORD)CW_USEDEFAULT,
+		(DWORD)CW_USEDEFAULT,
+		(DWORD)CW_USEDEFAULT,
+		(DWORD)0, (DWORD)0, (DWORD)0,
+		(DWORD)dwFlags,
+		ShowWindowFlags,
+		0, NULL,
+		NULL,
+		HANDLE(PipeWriteChild),
+		HANDLE(PipeWriteChild),
+		LPPROC_THREAD_ATTRIBUTE_LIST(AttributeList.GetData())
+	};
+
+	// create the child process
+	FString CommandLine = FString::Printf(TEXT("\"%s\" %s"), URL, Parms);
+	PROCESS_INFORMATION ProcInfo;
+
+	if (!CreateProcess(NULL, CommandLine.GetCharArray().GetData(), nullptr, nullptr, !AttributeList.IsEmpty(), (DWORD)CreateFlags, NULL, OptionalWorkingDirectory, &StartupInfoEx.StartupInfo, &ProcInfo))
+	{
+		HandleError();
+		if (!AttributeList.IsEmpty())
+		{
+			DeleteProcThreadAttributeList(LPPROC_THREAD_ATTRIBUTE_LIST(AttributeList.GetData()));
+		}
+		return {};
+	}
+
+	if (!AttributeList.IsEmpty())
+	{
+		DeleteProcThreadAttributeList(LPPROC_THREAD_ATTRIBUTE_LIST(AttributeList.GetData()));
+	}
+	
+	if (OutProcessID != nullptr)
+	{
+		*OutProcessID = ProcInfo.dwProcessId;
+	}
+
+	::CloseHandle(ProcInfo.hThread);
+
+	return FProcHandle(ProcInfo.hProcess);
+}
+#include "Windows/HideWindowsPlatformTypes.h"
+#else
+static FProcHandle CreateProc(const TCHAR* URL, const TCHAR* Parms, bool bLaunchDetached, bool bLaunchHidden, bool bLaunchReallyHidden, uint32* OutProcessID, int32 PriorityModifier, const TCHAR* OptionalWorkingDirectory, void* PipeWriteChild)
+{
+	return FPlatformProcess::CreateProc(URL, Parms, bLaunchDetached, bLaunchHidden, bLaunchReallyHidden, OutProcessID, PriorityModifier, OptionalWorkingDirectory, PipeWriteChild);
+}
+#endif
+
 #define LOCTEXT_NAMESPACE "FAndroidDeviceDetectionModule" 
 
 DEFINE_LOG_CATEGORY_STATIC(AndroidDeviceDetectionLog, Log, All);
@@ -159,7 +305,7 @@ private:
 		int32 ReturnCode = -1;
 		
 		{
-			FProcHandle Process = FPlatformProcess::CreateProc(*ADBPath, *CommandLine, true, true, false, nullptr, 0, nullptr, WritePipe, ReadPipe);
+			FProcHandle Process = CreateProc(*ADBPath, *CommandLine, true, true, false, nullptr, 0, nullptr, WritePipe);
 			if (Process.IsValid())
 			{
 				for (UE::FMonotonicTimePoint MaxTime = UE::FMonotonicTimePoint::Now() + UE::FMonotonicTimeSpan::FromSeconds(10.0);;)
@@ -181,7 +327,7 @@ private:
 
 					if (bStopRequested.load(std::memory_order_relaxed) || UE::FMonotonicTimePoint::Now() >= MaxTime)
 					{
-						FPlatformProcess::TerminateProc(Process, true);
+						FPlatformProcess::TerminateProc(Process);
 						break;
 					}
 				}
