Compile fix for Incremental build

#tests Editor runs

diff --git a/Engine/Source/Developer/Android/AndroidDeviceDetection/Private/AndroidDeviceDetectionModule.cpp b/Engine/Source/Developer/Android/AndroidDeviceDetection/Private/AndroidDeviceDetectionModule.cpp
--- a/Engine/Source/Developer/Android/AndroidDeviceDetection/Private/AndroidDeviceDetectionModule.cpp
+++ b/Engine/Source/Developer/Android/AndroidDeviceDetection/Private/AndroidDeviceDetectionModule.cpp
@@ -138,19 +138,21 @@ static FProcHandle CreateProc(const TCHAR* URL, const TCHAR* Parms, bool bLaunch
 	// initialize startup info
 	STARTUPINFOEX StartupInfoEx
 	{
-		sizeof(STARTUPINFOEX),
-		NULL, NULL, NULL,
-		(DWORD)CW_USEDEFAULT,
-		(DWORD)CW_USEDEFAULT,
-		(DWORD)CW_USEDEFAULT,
-		(DWORD)CW_USEDEFAULT,
-		(DWORD)0, (DWORD)0, (DWORD)0,
-		(DWORD)dwFlags,
-		ShowWindowFlags,
-		0, NULL,
-		NULL,
-		HANDLE(PipeWriteChild),
-		HANDLE(PipeWriteChild),
+		{
+			sizeof(STARTUPINFOEX),
+			NULL, NULL, NULL,
+			(DWORD)CW_USEDEFAULT,
+			(DWORD)CW_USEDEFAULT,
+			(DWORD)CW_USEDEFAULT,
+			(DWORD)CW_USEDEFAULT,
+			(DWORD)0, (DWORD)0, (DWORD)0,
+			(DWORD)dwFlags,
+			ShowWindowFlags,
+			0, NULL,
+			NULL,
+			HANDLE(PipeWriteChild),
+			HANDLE(PipeWriteChild)
+		},
 		LPPROC_THREAD_ATTRIBUTE_LIST(AttributeList.GetData())
 	};
 
