Fix android port forwarding plugin leaving stale zen server ports
#rnx

diff --git a/Engine/Config/BaseEngine.ini b/Engine/Config/BaseEngine.ini
--- a/Engine/Config/BaseEngine.ini
+++ b/Engine/Config/BaseEngine.ini
@@ -2613,6 +2613,9 @@ bIsSupported=True
 ; make full cooked editor builds, you could use {Platform}CookedEditor
 RuntimeTargetPlatform="{Platform}"
 
+[Zen.Plugin.AndroidPortForwarder]
+Name=Engine/Source/Developer/Android/AndroidZenServerPlugin/Binaries/Win64/AndroidPortForwarder2.dll
+
 [Zen]
 AutoLaunch=true
 
diff --git a/Engine/Source/Developer/Android/AndroidZenServerPlugin/AndroidPortForwarder/AndroidPortForwarder.cs b/Engine/Source/Developer/Android/AndroidZenServerPlugin/AndroidPortForwarder/AndroidPortForwarder.cs
--- a/Engine/Source/Developer/Android/AndroidZenServerPlugin/AndroidPortForwarder/AndroidPortForwarder.cs
+++ b/Engine/Source/Developer/Android/AndroidZenServerPlugin/AndroidPortForwarder/AndroidPortForwarder.cs
@@ -8,7 +8,7 @@ public sealed class AndroidPortForwarder : IDisposable
 	private AndroidPortForwarderAsync? ForwarderAsync;
 	private Task? StartMonitorTask;
 
-	public void StartMonitor(string ADBExecutable, uint Port)
+	public void StartMonitor(string ADBPath, uint? ADBServerPortOpt, uint ZenServerPort)
 	{
 		Console.WriteLine("Forwarder start requested");
 
@@ -17,7 +17,7 @@ public sealed class AndroidPortForwarder : IDisposable
 		ForwarderAsync = new AndroidPortForwarderAsync();
 
 		// We don't wait for startup to finish to avoid blocking current thread
-		StartMonitorTask = ForwarderAsync.StartMonitorAsync(ADBExecutable, Port, TokenSource.Token);
+		StartMonitorTask = ForwarderAsync.StartMonitorAsync(ADBPath, ADBServerPortOpt, ZenServerPort, TokenSource.Token);
 
 		Console.WriteLine($"Forwarder started");
 	}
diff --git a/Engine/Source/Developer/Android/AndroidZenServerPlugin/AndroidPortForwarder/AndroidPortForwarderAsync.cs b/Engine/Source/Developer/Android/AndroidZenServerPlugin/AndroidPortForwarder/AndroidPortForwarderAsync.cs
--- a/Engine/Source/Developer/Android/AndroidZenServerPlugin/AndroidPortForwarder/AndroidPortForwarderAsync.cs
+++ b/Engine/Source/Developer/Android/AndroidZenServerPlugin/AndroidPortForwarder/AndroidPortForwarderAsync.cs
@@ -5,40 +5,47 @@ namespace AndroidZenServerPlugin;
 using AdvancedSharpAdbClient;
 using AdvancedSharpAdbClient.DeviceCommands;
 using AdvancedSharpAdbClient.Models;
+using System.Diagnostics;
+using System.Net;
 
 public sealed class AndroidPortForwarderAsync : IDisposable
 {
 	private AdbClient? Client;
 	private DeviceMonitor? DeviceMonitor;
 
-	public async Task StartMonitorAsync(string ADBExecutable, uint Port, CancellationToken CancellationToken)
+	public async Task StartMonitorAsync(string ADBPath, uint? ADBServerPortOpt, uint ZenServerPort, CancellationToken CancellationToken)
 	{
 		try
 		{
-			if (string.IsNullOrEmpty(ADBExecutable))
+			if (string.IsNullOrEmpty(ADBPath))
 			{
-				ADBExecutable = Environment.ExpandEnvironmentVariables("%ANDROID_HOME%/platform-tools/adb" + (OperatingSystem.IsWindows() ? ".exe" : ""));
+				ADBPath = Environment.ExpandEnvironmentVariables("%ANDROID_HOME%/platform-tools/adb" + (OperatingSystem.IsWindows() ? ".exe" : ""));
 			}
 
-			await AdbServer.Instance.StartServerAsync(ADBExecutable, false, CancellationToken);
+			uint ResolvedADBServerPort = ADBServerPortOpt.GetValueOrDefault((uint)AdbClient.AdbServerPort);
+			Client = new AdbClient(new IPEndPoint(IPAddress.Loopback, (int)ResolvedADBServerPort));
 
-			Console.WriteLine($"Forwarder server started at '{ADBExecutable}'");
+			// inject custom command line client
+			AdbServer.Instance = new AdbServer(Client, Factories.AdbSocketFactory, Path => new DetachedAdbCommandLineClient(Path));
+
+			await AdbServer.Instance.StartServerAsync(ADBPath, false, CancellationToken);
+
+			Console.WriteLine($"Forwarder server started at '{ADBPath}'");
 
-			Client = new AdbClient();
 			DeviceMonitor = new DeviceMonitor(Client.CreateAdbSocket());
 
 			DeviceMonitor.DeviceConnected += async (_, args) =>
 			{
 				if (args.Device.State == DeviceState.Online)
 				{
-					await ForwardPortsToDeviceAsync(Client, args.Device, Port, CancellationToken);
+					await ForwardPortsToDeviceAsync(Client, args.Device, ZenServerPort, CancellationToken);
 				}
 			};
 			DeviceMonitor.DeviceChanged += async (_, args) =>
 			{
 				if (args.Device.State == DeviceState.Online)
 				{
-					await ForwardPortsToDeviceAsync(Client, args.Device, Port, CancellationToken);
+					await ForwardPortsToDeviceAsync(Client, args.Device, ZenServerPort, CancellationToken);
 				}
 			};
 
@@ -125,4 +132,77 @@ public sealed class AndroidPortForwarderAsync : IDisposable
 
 		return bAnyActivities;
 	}
+
+	// Right now launching adb with bInheritHandles makes adb server inherit zen server port handles,
+	// which prevents them from being released when zen server is shutdown/restarted.
+	// To circumvent this we launch adb commands via UseShellExecute.
+	// See [REDACTED_PROFILE_URL]/runtime/issues/13943
+	internal class DetachedAdbCommandLineClient(string AdbPath) : AdbCommandLineClient(AdbPath)
+	{
+		private bool ShouldRunAsDetached(string Filename, string Command, ICollection<string>? ErrorOutput, ICollection<string>? StandardOutput)
+		{
+			// always run start-server in detached mode without inheriting handles
+			return Command.EndsWith("start-server") ||
+			       (ErrorOutput == null && StandardOutput == null);
+		}
+
+		protected override int RunProcess(string Filename, string Command, ICollection<string>? ErrorOutput, ICollection<string>? StandardOutput, int Timeout)
+		{
+			if (!ShouldRunAsDetached(Filename, Command, ErrorOutput, StandardOutput))
+			{
+				return base.RunProcess(Filename, Command, ErrorOutput, StandardOutput, Timeout);
+			}
+
+			ProcessStartInfo Info = new()
+			{
+				FileName = Filename,
+				Arguments = Command,
+				WindowStyle = ProcessWindowStyle.Hidden,
+				UseShellExecute = true,
+			};
+
+			using Process Handle = Process.Start(Info) ?? throw new Exception($"Failed to start process '{Filename}' with command '{Command}'");
+
+			if (!Handle.WaitForExit(Timeout))
+			{
+				Handle.Kill();
+			}
+
+			return Handle.ExitCode;
+		}
+
+		protected override async Task<int> RunProcessAsync(string Filename, string Command, ICollection<string>? ErrorOutput, ICollection<string>? StandardOutput, CancellationToken CancellationToken = default)
+		{
+			if (!ShouldRunAsDetached(Filename, Command, ErrorOutput, StandardOutput))
+			{
+				return await base.RunProcessAsync(Filename, Command, ErrorOutput, StandardOutput, CancellationToken);
+			}
+
+			ProcessStartInfo Info = new()
+			{
+				FileName = Filename,
+				Arguments = Command,
+				WindowStyle = ProcessWindowStyle.Hidden,
+				UseShellExecute = true,
+			};
+
+			using Process Handle = Process.Start(Info) ?? throw new Exception($"Failed to start process '{Filename}' with command '{Command}'");
+
+			try
+			{
+				await Handle.WaitForExitAsync(CancellationToken);
+			}
+			catch (OperationCanceledException)
+			{
+				if (!Handle.HasExited)
+				{
+					Handle.Kill();
+				}
+
+				throw;
+			}
+
+			return Handle.ExitCode;
+		}
+	}
 }
diff --git a/Engine/Source/Developer/Android/AndroidZenServerPlugin/AndroidPortForwarder/AndroidPortForwarderCLI.cs b/Engine/Source/Developer/Android/AndroidZenServerPlugin/AndroidPortForwarder/AndroidPortForwarderCLI.cs
--- a/Engine/Source/Developer/Android/AndroidZenServerPlugin/AndroidPortForwarder/AndroidPortForwarderCLI.cs
+++ b/Engine/Source/Developer/Android/AndroidZenServerPlugin/AndroidPortForwarder/AndroidPortForwarderCLI.cs
@@ -20,7 +20,7 @@ Console.CancelKeyPress += (_, eventArgs) =>
 
 using AndroidPortForwarder forwarder = new AndroidPortForwarder();
 
-forwarder.StartMonitor(string.Empty, 8558);
+forwarder.StartMonitor(string.Empty, null, 8558);
 
 exitEvent.WaitOne();
 
diff --git a/Engine/Source/Developer/Android/AndroidZenServerPlugin/AndroidPortForwarder/ZenServerAdapter.cs b/Engine/Source/Developer/Android/AndroidZenServerPlugin/AndroidPortForwarder/ZenServerAdapter.cs
--- a/Engine/Source/Developer/Android/AndroidZenServerPlugin/AndroidPortForwarder/ZenServerAdapter.cs
+++ b/Engine/Source/Developer/Android/AndroidZenServerPlugin/AndroidPortForwarder/ZenServerAdapter.cs
@@ -10,7 +10,7 @@ public static class ZenServerAdapter
 	private static AndroidPortForwarder? AndroidPortForwarder;
 
 	[UnmanagedCallersOnly(EntryPoint = "StartAndroidPortForwarder")]
-	private static unsafe void StartAndroidPortForwarder(delegate*<void*, int, char*, void> LogCallback, void* UsrPtr, IntPtr ADBExecutablePtr, uint Port)
+	private static unsafe void StartAndroidPortForwarder(delegate*<void*, int, char*, void> LogCallback, void* UsrPtr, IntPtr ADBPathPtr, uint* ADBServerPortOptPtr, uint ZenServerPort)
 	{
 		ConsoleRedirector?.Dispose();
 		ConsoleRedirector = new ConsoleRedirector((Level, Line) =>
@@ -26,12 +26,13 @@ public static class ZenServerAdapter
 			}
 		});
 
-		string ADBExecutable = Marshal.PtrToStringAnsi(ADBExecutablePtr) ?? string.Empty;
+		string ADBPath = ADBPathPtr != IntPtr.Zero ? (Marshal.PtrToStringAnsi(ADBPathPtr) ?? string.Empty) : string.Empty;
+		uint? ADBServerPortOpt = ADBServerPortOptPtr != null ? *ADBServerPortOptPtr : null;
 
 		AndroidPortForwarder?.Dispose();
 		AndroidPortForwarder = new AndroidPortForwarder();
 
-		AndroidPortForwarder.StartMonitor(ADBExecutable, Port);
+		AndroidPortForwarder.StartMonitor(ADBPath, ADBServerPortOpt, ZenServerPort);
 	}
 
 	[UnmanagedCallersOnly(EntryPoint = "StopAndroidPortForwarder")]
diff --git a/Engine/Source/Developer/Android/AndroidZenServerPlugin/ZenServerAdapter/ZenServerAdapter.cpp b/Engine/Source/Developer/Android/AndroidZenServerPlugin/ZenServerAdapter/ZenServerAdapter.cpp
--- a/Engine/Source/Developer/Android/AndroidZenServerPlugin/ZenServerAdapter/ZenServerAdapter.cpp
+++ b/Engine/Source/Developer/Android/AndroidZenServerPlugin/ZenServerAdapter/ZenServerAdapter.cpp
@@ -2,19 +2,53 @@
 
 #include <transportplugin.h>
 #include <atomic>
-#include <cstring>
-#include <cstdlib>
-#include <cstdio>
+#include <cstdarg>
+#include <optional>
+#include <string>
+#include <vector>
 
+static const uint32_t kPluginVersion = 3;
+static const uint32_t kDefaultZenServerPort = 8558;
+
+// native -> managed API
 extern "C"
 {
 	typedef void (*LogCallback)(void* UsrPtr, int32_t Level, const char* Buffer);
 
-	// native -> managed API
-	void StartAndroidPortForwarder(LogCallback Callback, void* UsrPtr, const char* ADBExecutable, uint32_t Port);
+	void StartAndroidPortForwarder(LogCallback Callback, void* UsrPtr, const char* ADBPath, const uint32_t* ADBServerPortOpt, uint32_t ZenServerPort);
 	void StopAndroidPortForwarder();
 }
 
+// logging utility
+static void
+#if __clang__
+__attribute__((format(printf, 2, 3)))
+#endif
+LogMessage(zen::TransportLogger* Logger, const zen::TransportLogger::LogLevel Level, const char* Format, ...)
+{
+	char Buffer[1024];
+
+	va_list Args;
+	va_start(Args, Format);
+	const int RequiredSize = vsnprintf(Buffer, sizeof(Buffer), Format, Args);
+	if (RequiredSize < sizeof(Buffer))
+	{
+		Logger->LogMessage(Level, Buffer);
+	}
+	else
+	{
+		char* AllocBuffer = (char*)malloc(RequiredSize + 1);
+		vsnprintf(AllocBuffer, RequiredSize + 1, Format, Args);
+		Logger->LogMessage(Level, AllocBuffer);
+		free(AllocBuffer);
+	}
+	va_end(Args);
+}
+
+#define LOG_INFO(Logger, ...) LogMessage(Logger, zen::TransportLogger::LogLevel::Info, __VA_ARGS__)
+#define LOG_WARN(Logger, ...) LogMessage(Logger, zen::TransportLogger::LogLevel::Warn, __VA_ARGS__)
+#define LOG_ERR(Logger, ...)  LogMessage(Logger, zen::TransportLogger::LogLevel::Err,  __VA_ARGS__)
+
 class ZenServerAdapter : public zen::TransportPlugin
 {
 public:
@@ -43,31 +77,40 @@ public:
 	{
 		if (!strcmp(OptionTag, "port"))
 		{
-			char* EndPtr = nullptr;
-			const long PortValue = strtoul(OptionValue, &EndPtr, 10);
-			if (OptionValue != EndPtr)
+			if (!ParseOption(OptionValue, ZenServerPort))
+			{
+				LOG_ERR(Logger, "Can't parse zen server port value '%s'", OptionValue);
+			}
+		}
+		else if (!strcmp(OptionTag, "adb_server_port"))
+		{
+			uint32_t ServerPortTemp = 0;
+			if (ParseOption(OptionValue, ServerPortTemp))
 			{
-				Port = (uint32_t)PortValue;
+				ADBServerPort = ServerPortTemp;
 			}
 			else
 			{
-				Logger->LogMessage(zen::TransportLogger::LogLevel::Err, "Can't parse port value");
+				LOG_ERR(Logger, "Can't parse adb server port value '%s'", OptionValue);
 			}
 		}
-		else if (!strcmp(OptionTag, "adb_executable"))
+		else if (!strcmp(OptionTag, "adb_path"))
 		{
-			snprintf(ADBExecutable, sizeof(ADBExecutable), "%s", OptionValue);
+			ADBPath = std::string(OptionValue);
 		}
 	}
 
 	void Initialize(zen::TransportServer* ServerInterface) override
 	{
+		uint32_t ADBServerPortTemp = ADBServerPort.value_or(0);
+		const uint32_t* ADBServerPortTempPtr = ADBServerPort.has_value() ? &ADBServerPortTemp : nullptr;
+
 		StartAndroidPortForwarder(
 		[](void* ThisPtr, int32_t Level, const char* Buffer)
 		{
 			static_cast<ZenServerAdapter*>(ThisPtr)->Logger->LogMessage(Level == 1 ? zen::TransportLogger::LogLevel::Err : zen::TransportLogger::LogLevel::Info, Buffer);
 		},
-		this, ADBExecutable, Port);
+		this, ADBPath.has_value() ? ADBPath.value().c_str() : nullptr, ADBServerPortTempPtr, ZenServerPort);
 	}
 
 	void Shutdown() override
@@ -87,9 +130,24 @@ public:
 
 private:
 	std::atomic<uint32_t> ReferenceCount = {0};
-	uint32_t Port = 8558;
-	char ADBExecutable[1024] = {};
+	uint32_t ZenServerPort = kDefaultZenServerPort;
+	std::optional<uint32_t> ADBServerPort = std::optional<uint32_t>();
+	std::optional<std::string> ADBPath = std::optional<std::string>();
 	zen::TransportLogger* Logger = nullptr;
+
+	static bool ParseOption(const char* Option, uint32_t& Result)
+	{
+		char* EndPtr = nullptr;
+		const unsigned long Value = std::strtoul(Option, &EndPtr, 10);
+
+		if (Option == EndPtr || Value > UINT32_MAX)
+		{
+			return false;
+		}
+
+		Result = static_cast<uint32_t>(Value);
+		return true;
+	}
 };
 
 #if defined(_MSC_VER)
@@ -107,7 +165,7 @@ extern "C" DLL_TRANSPORT_API void GetTransportPluginVersion(uint32_t* OutApiVers
 
 	if (OutPluginVersion != nullptr)
 	{
-		*OutPluginVersion = 2;
+		*OutPluginVersion = kPluginVersion;
 	}
 }
 
@@ -115,3 +173,31 @@ extern "C" DLL_TRANSPORT_API zen::TransportPlugin* CreateTransportPlugin(zen::Tr
 {
 	return new ZenServerAdapter(Logger);
 }
+
+#if TEST_CLI
+
+extern "C"
+{
+	void StartAndroidPortForwarder(LogCallback Callback, void* UsrPtr, const char* ADBServerPortOpt, const uint32_t* ADBServerPortOpt, uint32_t ZenServerPort) {}
+	void StopAndroidPortForwarder() {}
+}
+
+class StubTransportLogger : public zen::TransportLogger
+{
+public:
+	void LogMessage(LogLevel Level, const char* Message) override { printf("%u %s\n", (uint8_t)Level, Message); }
+};
+
+int main()
+{
+	auto Plugin = CreateTransportPlugin(new StubTransportLogger);
+	Plugin->Configure("port", "abc");
+	Plugin->Initialize(nullptr);
+	while (true)
+	{
+		Sleep(100);
+	}
+	return 0;
+}
+
+#endif
\ No newline at end of file
