Fix FAndroidTargetDevice::PowerOn() Hanging on Unix/Apple Computers

#tests Tested powering on the emulator on both Windows and Mac
#rnx

diff --git a/Engine/Source/Developer/Android/AndroidTargetPlatformControls/Private/AndroidTargetDevice.inl b/Engine/Source/Developer/Android/AndroidTargetPlatformControls/Private/AndroidTargetDevice.inl
--- a/Engine/Source/Developer/Android/AndroidTargetPlatformControls/Private/AndroidTargetDevice.inl
+++ b/Engine/Source/Developer/Android/AndroidTargetPlatformControls/Private/AndroidTargetDevice.inl
@@ -48,7 +48,12 @@ inline bool FAndroidTargetDevice::PowerOn()
 
 	if (FString AndroidHomeDirectory = FPlatformMisc::GetEnvironmentVariable(TEXT("ANDROID_HOME")); !AndroidHomeDirectory.IsEmpty())
 	{
-		FProcHandle Process = FPlatformProcess::CreateProc(*FPaths::Combine(AndroidHomeDirectory, TEXT("emulator"), TEXT("emulator")), *(TEXT("-avd ") + DeviceId.RightChop(4)), false, false, true, nullptr, 0, nullptr, nullptr);
+#if PLATFORM_WINDOWS
+		bool bLaunchDetached = false;
+#else
+		bool bLaunchDetached = true;
+#endif
+		FProcHandle Process = FPlatformProcess::CreateProc(*FPaths::Combine(AndroidHomeDirectory, TEXT("emulator"), TEXT("emulator")), *(TEXT("-avd ") + DeviceId.RightChop(4)), bLaunchDetached, false, true, nullptr, 0, nullptr, nullptr);
 
 		if (Process.IsValid())
 		{
