Fixes for Android when bSupportsDistanceFields is enabled in DDPI
Override UsesDistanceFields for AndroidTPS like IOSTPS to return the correct value, bDistanceField
r.DistanceFields  should only be set in the /Script/Engine.RendererSettings, not in DP

diff --git a/Engine/Config/BaseDeviceProfiles.ini b/Engine/Config/BaseDeviceProfiles.ini
--- a/Engine/Config/BaseDeviceProfiles.ini
+++ b/Engine/Config/BaseDeviceProfiles.ini
@@ -891,7 +891,6 @@ bIsVisibleForAssets=True
 +CVars=r.Android.DisableVulkanSM5Support=1
 ; PF_B8G8R8A8
 +CVars=r.DefaultBackBufferPixelFormat=0 
-+CVars=r.DistanceFields=0
 ; LightGrid Settings
 +CVars=r.Forward.LightLinkedListCulling=0
 ; Scale the LODDistanceFactor when doing a sky light capture
@@ -1012,7 +1011,6 @@ BaseProfileName=Android
 +CVars=r.SSR.Quality=0
 +CVars=r.Android.DisableVulkanSM5Support=0
 +CVars=r.Android.DisableVulkanSupport=0
-+CVars=r.DistanceFields=1
 +CVars=r.Vulkan.RayTracing.AllowCompaction=0
 +CVars=r.Vulkan.RayTracing.TLASPreferFastTraceTLAS=0
 
diff --git a/Engine/Source/Developer/Android/AndroidTargetPlatformSettings/Public/AndroidTargetPlatformSettings.h b/Engine/Source/Developer/Android/AndroidTargetPlatformSettings/Public/AndroidTargetPlatformSettings.h
--- a/Engine/Source/Developer/Android/AndroidTargetPlatformSettings/Public/AndroidTargetPlatformSettings.h
+++ b/Engine/Source/Developer/Android/AndroidTargetPlatformSettings/Public/AndroidTargetPlatformSettings.h
@@ -133,9 +133,13 @@ public:
 	ANDROIDTARGETPLATFORMSETTINGS_API bool SupportsES31() const;
 	ANDROIDTARGETPLATFORMSETTINGS_API bool SupportsVulkan() const;
 	ANDROIDTARGETPLATFORMSETTINGS_API bool SupportsVulkanSM5() const;
-protected:
 
-	
+	virtual bool UsesDistanceFields() const override
+	{
+		return bDistanceField;
+	}
+
+protected:
 
 	// r.Mobile.ShadingPath value
 	int32 MobileShadingPath;
