Clarify Android SDK API Level formatting documentation

#android
#13635

diff --git a/Engine/Source/Developer/Android/AndroidPlatformEditor/Private/AndroidSDKSettings.h b/Engine/Source/Developer/Android/AndroidPlatformEditor/Private/AndroidSDKSettings.h
--- a/Engine/Source/Developer/Android/AndroidPlatformEditor/Private/AndroidSDKSettings.h
+++ b/Engine/Source/Developer/Android/AndroidPlatformEditor/Private/AndroidSDKSettings.h
@@ -34,8 +34,8 @@ public:
 	UPROPERTY(GlobalConfig, EditAnywhere, Category = SDKConfig, Meta = (DisplayName = "Location of JAVA (the directory usually contains 'jdk')"))
 	FDirectoryPath JavaPath;
 	
-	// Which SDK to package and compile Java with (a specific version or (without quotes) 'latest' for latest version on disk, or 'matchndk' to match the NDK API Level)
-	UPROPERTY(GlobalConfig, EditAnywhere, Category = SDKConfig, Meta = (DisplayName = "SDK API Level (specific version, 'latest', or 'matchndk' - see tooltip)"))
+	// Which SDK to package and compile Java with (a specific version, usually (without quotes): 'android-NN' or (without quotes) 'latest' for latest version on disk, or 'matchndk' to match the NDK API Level)
+	UPROPERTY(GlobalConfig, EditAnywhere, Category = SDKConfig, Meta = (DisplayName = "SDK API Level (specific version e.g. 'android-NN', 'latest', or 'matchndk' - see tooltip)"))
 	FString SDKAPILevel;
 
 	// Which NDK to compile with (a specific version or (without quotes) 'latest' for latest version on disk). Note that choosing android-21 or later won't run on pre-5.0 devices.
diff --git a/Engine/Source/Programs/UnrealBuildTool/Platform/Android/UEDeployAndroid.cs b/Engine/Source/Programs/UnrealBuildTool/Platform/Android/UEDeployAndroid.cs
--- a/Engine/Source/Programs/UnrealBuildTool/Platform/Android/UEDeployAndroid.cs
+++ b/Engine/Source/Programs/UnrealBuildTool/Platform/Android/UEDeployAndroid.cs
@@ -242,16 +242,14 @@ namespace UnrealBuildTool
 
 		private int GetApiLevelInt(string ApiString)
 		{
-			int VersionInt = 0;
-			if (ApiString.Contains('-'))
+			int DashIndex = ApiString.LastIndexOf('-');
+			if (DashIndex != -1 && int.TryParse(ApiString.Substring(DashIndex + 1), out int Version))
 			{
-				int Version;
-				if (Int32.TryParse(ApiString.Substring(ApiString.LastIndexOf('-') + 1), out Version))
-				{
-					VersionInt = Version;
-				}
+				return Version;
 			}
-			return VersionInt;
+			
+			Logger.LogWarning("Unable to parse SDK API level '{ApiString}', it is expected to be of the form 'android-NN'", ApiString);
+			return 0;
 		}
 
 		private string? CachedSDKLevel = null;
