Use GetPlatformValueVariable and not GetConfigSystemin IOS and Android TPS

diff --git a/Engine/Source/Developer/Android/AndroidTargetPlatformSettings/Private/AndroidTargetPlatformSettings.cpp b/Engine/Source/Developer/Android/AndroidTargetPlatformSettings/Private/AndroidTargetPlatformSettings.cpp
--- a/Engine/Source/Developer/Android/AndroidTargetPlatformSettings/Private/AndroidTargetPlatformSettings.cpp
+++ b/Engine/Source/Developer/Android/AndroidTargetPlatformSettings/Private/AndroidTargetPlatformSettings.cpp
@@ -9,19 +9,72 @@
 
 FAndroidTargetPlatformSettings::FAndroidTargetPlatformSettings(const TCHAR* CookFlavor, const TCHAR* OverrideIniPlatformName)
 	: TTargetPlatformSettingsBase(CookFlavor, OverrideIniPlatformName)
-	, MobileShadingPath(0)
-	, bDistanceField(false)
-	, bMobileForwardEnableClusteredReflections(false)
 {
 #if WITH_ENGINE
 	TextureLODSettings = nullptr; // These are registered by the device profile system.
 	StaticMeshLODSettings.Initialize(this);
-	GetConfigSystem()->GetBool(TEXT("/Script/Engine.RendererSettings"), TEXT("r.DistanceFields"), bDistanceField, GEngineIni);
-	GetConfigSystem()->GetInt(TEXT("/Script/Engine.RendererSettings"), TEXT("r.Mobile.ShadingPath"), MobileShadingPath, GEngineIni);
-	GetConfigSystem()->GetBool(TEXT("/Script/Engine.RendererSettings"), TEXT("r.Mobile.Forward.EnableClusteredReflections"), bMobileForwardEnableClusteredReflections, GEngineIni);
 #endif
 }
 
+// Using GetConfigSystem instead of GetPlatformValueVariable here because RenderUtils calls this too early.
+bool FAndroidTargetPlatformSettings::GetIsDistanceField() const
+{
+	static TOptional<bool> bDistanceField;
+	if (!bDistanceField.IsSet())
+	{
+		bDistanceField = false;
+#if WITH_ENGINE
+		static IConsoleVariable* CvarDistanceFields = IConsoleManager::Get().FindConsoleVariable(TEXT("r.DistanceFields"));
+		if (CvarDistanceFields != nullptr)
+		{
+			bDistanceField = CvarDistanceFields->GetInt() != 0;
+		}
+		
+		GetConfigSystem()->GetBool(TEXT("/Script/Engine.RendererSettings"), TEXT("r.DistanceFields"), bDistanceField.GetValue(), GEngineIni);
+#endif
+	}
+	return bDistanceField.GetValue();
+}
+
+int32 FAndroidTargetPlatformSettings::GetMobileShadingPath() const
+{
+	static TOptional<int32> MobileShadingPath;
+	if (!MobileShadingPath.IsSet())
+	{
+		MobileShadingPath = 0;
+#if WITH_ENGINE
+		static IConsoleVariable* CvarShadingPath = IConsoleManager::Get().FindConsoleVariable(TEXT("r.Mobile.ShadingPath"));
+		if (CvarShadingPath != nullptr)
+		{
+			MobileShadingPath = CvarShadingPath->GetPlatformValueVariable(*IniPlatformName())->GetInt();
+		}
+#endif
+	}
+	return MobileShadingPath.GetValue();
+}
+
+bool FAndroidTargetPlatformSettings::GetIsMobileForwardEnableClusteredReflections() const
+{
+	static TOptional<bool> bMobileForwardEnableClusteredReflections;
+	if (!bMobileForwardEnableClusteredReflections.IsSet())
+	{
+		bMobileForwardEnableClusteredReflections = false;
+#if WITH_ENGINE
+		static IConsoleVariable* CvarFwdEnableClusteredReflections = IConsoleManager::Get().FindConsoleVariable(TEXT("r.Mobile.Forward.EnableClusteredReflections"));
+		if (CvarFwdEnableClusteredReflections != nullptr)
+		{
+			bMobileForwardEnableClusteredReflections = CvarFwdEnableClusteredReflections->GetPlatformValueVariable(*IniPlatformName())->GetInt() != 0;
+		}
+#endif
+	}
+	return bMobileForwardEnableClusteredReflections.GetValue();
+}
+
+bool FAndroidTargetPlatformSettings::UsesDistanceFields() const
+{
+	return GetIsDistanceField();
+}
+
 bool FAndroidTargetPlatformSettings::SupportsES31() const
 {
 	// default no support for ES31
@@ -92,9 +145,9 @@ bool FAndroidTargetPlatformSettings::SupportsFeature(ETargetPlatformFeatures Fea
 #if WITH_ENGINE
 void FAndroidTargetPlatformSettings::GetReflectionCaptureFormats(TArray<FName>& OutFormats) const
 {
-	const bool bMobileDeferredShading = (MobileShadingPath == 1);
+	const bool bMobileDeferredShading = (GetMobileShadingPath() == 1);
 
-	if (SupportsVulkanSM5() || bMobileDeferredShading || bMobileForwardEnableClusteredReflections)
+	if (SupportsVulkanSM5() || bMobileDeferredShading || GetIsMobileForwardEnableClusteredReflections())
 	{
 		// use Full HDR with SM5 and Mobile Deferred
 		OutFormats.Add(FName(TEXT("FullHDR")));
diff --git a/Engine/Source/Developer/Android/AndroidTargetPlatformSettings/Public/AndroidTargetPlatformSettings.h b/Engine/Source/Developer/Android/AndroidTargetPlatformSettings/Public/AndroidTargetPlatformSettings.h
--- a/Engine/Source/Developer/Android/AndroidTargetPlatformSettings/Public/AndroidTargetPlatformSettings.h
+++ b/Engine/Source/Developer/Android/AndroidTargetPlatformSettings/Public/AndroidTargetPlatformSettings.h
@@ -134,22 +134,9 @@ public:
 	ANDROIDTARGETPLATFORMSETTINGS_API bool SupportsVulkan() const;
 	ANDROIDTARGETPLATFORMSETTINGS_API bool SupportsVulkanSM5() const;
 
-	virtual bool UsesDistanceFields() const override
-	{
-		return bDistanceField;
-	}
+	virtual bool UsesDistanceFields() const override;
 
 protected:
-
-	// r.Mobile.ShadingPath value
-	int32 MobileShadingPath;
-
-	// true if DistanceField is enabled
-	bool bDistanceField;
-
-	// r.Mobile.Forward.EnableClusteredReflections value
-	bool bMobileForwardEnableClusteredReflections;
-
 #if WITH_ENGINE
 	// Holds a cache of the target LOD settings.
 	const UTextureLODSettings* TextureLODSettings;
@@ -157,6 +144,10 @@ protected:
 	// Holds the static mesh LOD settings.
 	FStaticMeshLODSettings StaticMeshLODSettings;
 #endif
+private:
+	bool GetIsDistanceField() const;
+	int32 GetMobileShadingPath() const;
+	bool GetIsMobileForwardEnableClusteredReflections() const;
 };
 
 class FAndroid_DXTTargetPlatformSettings : public FAndroidTargetPlatformSettings
diff --git a/Engine/Source/Developer/IOS/IOSTargetPlatformSettings/Private/IOSTargetPlatformSettings.cpp b/Engine/Source/Developer/IOS/IOSTargetPlatformSettings/Private/IOSTargetPlatformSettings.cpp
--- a/Engine/Source/Developer/IOS/IOSTargetPlatformSettings/Private/IOSTargetPlatformSettings.cpp
+++ b/Engine/Source/Developer/IOS/IOSTargetPlatformSettings/Private/IOSTargetPlatformSettings.cpp
@@ -23,19 +23,71 @@ FIOSTargetPlatformSettings::FIOSTargetPlatformSettings(bool bInIsTVOS, bool bInI
 	: TTargetPlatformSettingsBase(nullptr, bInIsVisionOS ? TEXT("VisionOS") : nullptr)
 	, bIsTVOS(bInIsTVOS)
 	, bIsVisionOS(bInIsVisionOS)
-	, MobileShadingPath(0)
-	, bDistanceField(false)
-	, bMobileForwardEnableClusteredReflections(false)
 {
 #if WITH_ENGINE
 	TextureLODSettings = nullptr; // TextureLODSettings are registered by the device profile.
 	StaticMeshLODSettings.Initialize(this);
-	GetConfigSystem()->GetBool(TEXT("/Script/Engine.RendererSettings"), TEXT("r.DistanceFields"), bDistanceField, GEngineIni);
-	GetConfigSystem()->GetInt(TEXT("/Script/Engine.RendererSettings"), TEXT("r.Mobile.ShadingPath"), MobileShadingPath, GEngineIni);
-	GetConfigSystem()->GetBool(TEXT("/Script/Engine.RendererSettings"), TEXT("r.Mobile.Forward.EnableClusteredReflections"), bMobileForwardEnableClusteredReflections, GEngineIni);
 #endif // #if WITH_ENGINE
 }
 
+// Using GetConfigSystem instead of GetPlatformValueVariable here because RenderUtils calls this too early.
+bool FIOSTargetPlatformSettings::GetIsDistanceField() const
+{
+	static TOptional<bool> bDistanceField;
+	if (!bDistanceField.IsSet())
+	{
+		bDistanceField = false;
+#if WITH_ENGINE
+		static IConsoleVariable* CvarDistanceFields = IConsoleManager::Get().FindConsoleVariable(TEXT("r.DistanceFields"));
+		if (CvarDistanceFields != nullptr)
+		{
+			bDistanceField = CvarDistanceFields->GetInt() != 0;
+		}
+
+		GetConfigSystem()->GetBool(TEXT("/Script/Engine.RendererSettings"), TEXT("r.DistanceFields"), bDistanceField.GetValue(), GEngineIni);
+#endif
+	}
+	return bDistanceField.GetValue();
+}
+
+int32 FIOSTargetPlatformSettings::GetMobileShadingPath() const
+{
+	static TOptional<int32> MobileShadingPath;
+	if (!MobileShadingPath.IsSet())
+	{
+		MobileShadingPath = 0;
+#if WITH_ENGINE
+		static IConsoleVariable* CvarShadingPath = IConsoleManager::Get().FindConsoleVariable(TEXT("r.Mobile.ShadingPath"));
+		if (CvarShadingPath != nullptr)
+		{
+			MobileShadingPath = CvarShadingPath->GetPlatformValueVariable(*IniPlatformName())->GetInt();
+		}
+#endif
+	}
+	return MobileShadingPath.GetValue();
+}
+
+bool FIOSTargetPlatformSettings::GetIsMobileForwardEnableClusteredReflections() const
+{
+	static TOptional<bool> bMobileForwardEnableClusteredReflections;
+	if (!bMobileForwardEnableClusteredReflections.IsSet())
+	{
+		bMobileForwardEnableClusteredReflections = false;
+#if WITH_ENGINE
+		static IConsoleVariable* CvarFwdEnableClusteredReflections = IConsoleManager::Get().FindConsoleVariable(TEXT("r.Mobile.Forward.EnableClusteredReflections"));
+		if (CvarFwdEnableClusteredReflections != nullptr)
+		{
+			bMobileForwardEnableClusteredReflections = CvarFwdEnableClusteredReflections->GetPlatformValueVariable(*IniPlatformName())->GetInt() != 0;
+		}
+#endif
+	}
+	return bMobileForwardEnableClusteredReflections.GetValue();
+}
+
+bool FIOSTargetPlatformSettings::UsesDistanceFields() const
+{
+	return GetIsDistanceField();
+}
 
 FIOSTargetPlatformSettings::~FIOSTargetPlatformSettings()
 {
@@ -156,9 +208,9 @@ void FIOSTargetPlatformSettings::GetAllTargetedShaderFormats( TArray<FName>& Out
 
 void FIOSTargetPlatformSettings::GetReflectionCaptureFormats( TArray<FName>& OutFormats ) const
 {
-	const bool bMobileDeferredShading = (MobileShadingPath == 1);
+	const bool bMobileDeferredShading = (GetMobileShadingPath() == 1);
 
-	if (SupportsMetalMRT() || bMobileDeferredShading || bMobileForwardEnableClusteredReflections)
+	if (SupportsMetalMRT() || bMobileDeferredShading || GetIsMobileForwardEnableClusteredReflections())
 	{
 		OutFormats.Add(FName(TEXT("FullHDR")));
 	}
diff --git a/Engine/Source/Developer/IOS/IOSTargetPlatformSettings/Private/IOSTargetPlatformSettings.h b/Engine/Source/Developer/IOS/IOSTargetPlatformSettings/Private/IOSTargetPlatformSettings.h
--- a/Engine/Source/Developer/IOS/IOSTargetPlatformSettings/Private/IOSTargetPlatformSettings.h
+++ b/Engine/Source/Developer/IOS/IOSTargetPlatformSettings/Private/IOSTargetPlatformSettings.h
@@ -59,26 +59,13 @@ public:
 
 	//~ Begin ITargetPlatform Interface
 
-	virtual bool UsesDistanceFields() const override
-	{
-		return bDistanceField;
-	}
+	virtual bool UsesDistanceFields() const override;
 
 private:
 	
 	// true if this is targeting TVOS vs IOS
 	bool bIsTVOS;
 	bool bIsVisionOS;
-
-	// r.Mobile.ShadingPath value
-	int32 MobileShadingPath;
-
-	// true if DistanceField is enabled
-	bool bDistanceField;
-
-	// r.Mobile.Forward.EnableClusteredReflections value
-	bool bMobileForwardEnableClusteredReflections;
-
 #if WITH_ENGINE
 	// Holds the cache of the target LOD settings.
 	const UTextureLODSettings* TextureLODSettings;
@@ -86,4 +73,7 @@ private:
 	// Holds the static mesh LOD settings.
 	FStaticMeshLODSettings StaticMeshLODSettings;
 #endif // WITH_ENGINE
+	bool GetIsDistanceField() const;
+	int32 GetMobileShadingPath() const ;
+	bool GetIsMobileForwardEnableClusteredReflections() const;
 };
