Account for spaces in the adb devices -l output when gathering model: and device: info

diff --git a/Engine/Source/Developer/Android/AndroidDeviceDetection/Private/AndroidDeviceDetectionModule.cpp b/Engine/Source/Developer/Android/AndroidDeviceDetection/Private/AndroidDeviceDetectionModule.cpp
--- a/Engine/Source/Developer/Android/AndroidDeviceDetection/Private/AndroidDeviceDetectionModule.cpp
+++ b/Engine/Source/Developer/Android/AndroidDeviceDetection/Private/AndroidDeviceDetectionModule.cpp
@@ -511,7 +511,14 @@ private:
 					NewDeviceInfo.GLESVersion = FCString::Atoi(*GLESVersionString);
 
 					// Find the device model
-					FParse::Value(*DeviceString, TEXT("model:"), NewDeviceInfo.Model);
+					{
+						FRegexPattern Pattern(TEXT("\\smodel:(.*?)(?=\\s+\\w+:|$)"));
+						FRegexMatcher Matcher(Pattern, DeviceString);
+						if (Matcher.FindNext())
+						{
+							NewDeviceInfo.Model = Matcher.GetCaptureGroup(1).TrimStartAndEnd();
+						}
+					}
 					// find the product model (this must match java's android.os.build.model)
 					FString ModelCommand = FString::Printf(TEXT("-s %s %s ro.product.model"), *SerialNumber, *GetPropCommand);
 					FString RoProductModel;
@@ -649,7 +656,15 @@ private:
 					}
 
 					// parse the device name
-					FParse::Value(*DeviceString, TEXT("device:"), NewDeviceInfo.DeviceName);
+					{
+						FRegexPattern Pattern(TEXT("\\sdevice:(.*?)(?=\\s+\\w+:|$)"));
+						FRegexMatcher Matcher(Pattern, DeviceString);
+						if (Matcher.FindNext())
+						{
+							NewDeviceInfo.DeviceName = Matcher.GetCaptureGroup(1).TrimStartAndEnd();
+						}
+					}
+					
 					if (NewDeviceInfo.DeviceName.IsEmpty())
 					{
 						FString DeviceCommand = FString::Printf(TEXT("-s %s %s ro.product.device"), *SerialNumber, *GetPropCommand);
