[AndroidDeviceDetection]
- Fix benign race on relaxed read of a boolean on a different thread
- Fix coding standard

diff --git a/Engine/Source/Developer/Android/AndroidDeviceDetection/Private/AndroidDeviceDetectionModule.cpp b/Engine/Source/Developer/Android/AndroidDeviceDetection/Private/AndroidDeviceDetectionModule.cpp
--- a/Engine/Source/Developer/Android/AndroidDeviceDetection/Private/AndroidDeviceDetectionModule.cpp
+++ b/Engine/Source/Developer/Android/AndroidDeviceDetection/Private/AndroidDeviceDetectionModule.cpp
@@ -64,8 +64,8 @@ public:
 		DeviceMap(InDeviceMap),
 		DeviceMapLock(InDeviceMapLock),
 		ADBPathCheckLock(InADBPathCheckLock),
-		HasADBPath(false),
-		ForceCheck(false)
+		bHasADBPath(false),
+		bForceCheck(false)
 	{
 		TcpMessagingModule = FModuleManager::LoadModulePtr<ITcpMessagingModule>("TcpMessaging");
 	}
@@ -98,17 +98,17 @@ public:
 		for (int32 LoopCount = 10;;)
 		{
 			// query when we have waited 'GAndroidDeviceDetectionPollInterval' seconds.
-			if (LoopCount++ >= GAndroidDeviceDetectionPollInterval || ForceCheck)
+			if (LoopCount++ >= GAndroidDeviceDetectionPollInterval || bForceCheck.load(std::memory_order_relaxed))
 			{
 				// Make sure we have an ADB path before checking
 				FScopeLock PathLock(ADBPathCheckLock);
-				if (HasADBPath)
+				if (bHasADBPath)
 				{
 					QueryConnectedDevices();
 				}
 
 				LoopCount = 0;
-				ForceCheck = false;
+				bForceCheck = false;
 			}
 
 			{
@@ -130,12 +130,12 @@ public:
 		GetPropCommand = MoveTemp(InGetPropCommand);
 		bGetExtensionsViaSurfaceFlinger = InbGetExtensionsViaSurfaceFlinger;
 
-		HasADBPath = !ADBPath.IsEmpty();
+		bHasADBPath = !ADBPath.IsEmpty();
 		// Force a check next time we go around otherwise it can take over 10sec to find devices
-		ForceCheck = HasADBPath;	
+		bForceCheck = bHasADBPath;
 
 		// If we have no path then clean the existing devices out
-		if (!HasADBPath && DeviceMap.Num() > 0)
+		if (!bHasADBPath && DeviceMap.Num() > 0)
 		{
 			DeviceMap.Reset();
 		}
@@ -885,8 +885,8 @@ private:
 	FCriticalSection* DeviceMapLock;
 
 	FCriticalSection* ADBPathCheckLock;
-	bool HasADBPath;
-	bool ForceCheck;
+	bool bHasADBPath;
+	std::atomic_bool bForceCheck;
 
 	ITcpMessagingModule* TcpMessagingModule;
 };
