Deprecated the entry "Android Graphics Debugger" from Project Settings->Platforms->Android->Graphics Debugger because the tools it refers to have been deprecated and are no longer available.
#Synced-CL 49758588

diff --git a/Engine/Source/Developer/Android/AndroidPlatformEditor/Private/AndroidTargetSettingsCustomization.cpp b/Engine/Source/Developer/Android/AndroidPlatformEditor/Private/AndroidTargetSettingsCustomization.cpp
--- a/Engine/Source/Developer/Android/AndroidPlatformEditor/Private/AndroidTargetSettingsCustomization.cpp
+++ b/Engine/Source/Developer/Android/AndroidPlatformEditor/Private/AndroidTargetSettingsCustomization.cpp
@@ -82,7 +82,6 @@ void FAndroidTargetSettingsCustomization::CustomizeDetails(IDetailLayoutBuilder&
 	BuildAppManifestSection(DetailLayout);
 	BuildIconSection(DetailLayout);
 	BuildLaunchImageSection(DetailLayout);
-	BuildGraphicsDebuggerSection(DetailLayout);
 	AudioPluginWidgetManager.BuildAudioCategory(DetailLayout, FString(TEXT("Android")));
 }
 
@@ -582,144 +581,6 @@ void FAndroidTargetSettingsCustomization::OnAppIDModified()
 	}
 }
 
-static EVisibility GraphicsDebuggerSettingsVisibility(EAndroidGraphicsDebugger::Type DebuggerType, TSharedPtr<IPropertyHandle> AndroidGraphicsDebuggerProperty)
-{
-	uint8 ValueAsByte = 0;
-	FPropertyAccess::Result Result = AndroidGraphicsDebuggerProperty->GetValue(ValueAsByte);
-	if (Result == FPropertyAccess::Success && ValueAsByte == static_cast<uint8>(DebuggerType))
-	{
-		return EVisibility::Visible;
-	}
-
-	return EVisibility::Hidden;
-}
-
-static FText GetMaliGraphicsDebuggerHelpText()
-{
-	const static FText InstallText(LOCTEXT("MGDInstallText", "Run the following command from a host command line from the target/unrooted directory located in the installation directory of the MGD tool, to install the MGD Daemon application on your device."));
-	const static FString InstallCommand(TEXT("adb install -r MGD.apk"));
-	const static FText RunText1(LOCTEXT("MGDIRunText1", "Run the following command from your host to establish a tunnel between your PC and the MGD Daemon. This needs to be done each time you connect your device by USB."));
-	const static FString RunCommand(TEXT("adb forward tcp:5002 tcp:5002"));
-	const static FText RunText2(LOCTEXT("MGDIRunText2", "Next, ensure you are running the daemon. Run the MGD Daemon application and switch it to the \"ON\" state"));
-	
-	FFormatOrderedArguments Args;
-	Args.Add(InstallText);
-	Args.Add(FText::FromString(InstallCommand));
-	Args.Add(RunText1);
-	Args.Add(FText::FromString(RunCommand));
-	Args.Add(RunText2);
-
-	return FText::Format(LOCTEXT("MaliGraphicsDebuggerHelpText","<RichTextBlock.TextHighlight>Installation</>\n{0}\n{1}\n\n<RichTextBlock.TextHighlight>Run</>\n{2}\n{3}\n{4}"), 
-		Args);
-}
-
-static FText GetAdrenoProfilerHelpText()
-{
-	const static FText RunText(LOCTEXT("AdrenoRunText", "Before profiling, and after rebooting your Android device, you must enable debug mode by setting the following property from the command line:"));
-	const static FString RunCommand(TEXT("adb shell setprop debug.egl.profiler 1"));
-		
-	FFormatOrderedArguments Args;
-	Args.Add(RunText);
-	Args.Add(FText::FromString(RunCommand));
-
-	return FText::Format(LOCTEXT("AdrenoHelpText","{0}\n{1}"), Args);
-}
-
-void FAndroidTargetSettingsCustomization::BuildGraphicsDebuggerSection(IDetailLayoutBuilder& DetailLayout)
-{
-	IDetailCategoryBuilder& GraphicsDebuggerCategory = DetailLayout.EditCategory(TEXT("GraphicsDebugger"));
-
-	TSharedPtr<IPropertyHandle> AndroidGraphicsDebuggerProperty = DetailLayout.GetProperty(GET_MEMBER_NAME_CHECKED(UAndroidRuntimeSettings, AndroidGraphicsDebugger));
-	GraphicsDebuggerCategory.AddProperty(AndroidGraphicsDebuggerProperty);
-
-	// Mali Graphics Debugger settings
-	{
-		TAttribute<EVisibility> MaliSettingsVisibility(
-			TAttribute<EVisibility>::Create(TAttribute<EVisibility>::FGetter::CreateStatic(GraphicsDebuggerSettingsVisibility, EAndroidGraphicsDebugger::Mali, AndroidGraphicsDebuggerProperty))
-		);
-		
-		TSharedPtr<IPropertyHandle> MaliGraphicsDebuggerPathProperty = DetailLayout.GetProperty(GET_MEMBER_NAME_CHECKED(UAndroidRuntimeSettings, MaliGraphicsDebuggerPath));
-		DetailLayout.HideProperty(MaliGraphicsDebuggerPathProperty);
-		GraphicsDebuggerCategory.AddProperty(MaliGraphicsDebuggerPathProperty).Visibility(MaliSettingsVisibility);
-
-		FText MGDHelpText = GetMaliGraphicsDebuggerHelpText();
-
-		GraphicsDebuggerCategory.AddCustomRow(LOCTEXT("MaliGraphicsDebuggerInfo", "Mali Graphics Debugger Info"), false)
-		.Visibility(MaliSettingsVisibility)
-		.WholeRowWidget
-		[
-			SNew(SBorder)
-			.Padding(1)
-			[
-				SNew(SVerticalBox)
-				+ SVerticalBox::Slot()
-				.Padding(FMargin(10, 10, 10, 10))
-				.AutoHeight()
-				[
-					SNew(SRichTextBlock)
-					.Text(MGDHelpText)
-					.TextStyle(FAppStyle::Get(), "MessageLog")
-					.DecoratorStyleSet(&FAppStyle::Get())
-					.AutoWrapText(true)
-				]
-				+ SVerticalBox::Slot()
-				.AutoHeight()
-				.Padding(FMargin(10, 10, 10, 10))
-				[	
-					SNew(SBox)
-					.HAlign(HAlign_Left)
-					[
-						SNew(SHyperlinkLaunchURL, TEXT("http://malideveloper.arm.com/resources/tools/mali-graphics-debugger/"))
-						.Text(LOCTEXT("MaliGraphicsDebuggerPage", "Mali Graphics Debugger home page"))
-						.ToolTipText(LOCTEXT("MaliGraphicsDebuggerPageTooltip", "Opens the Mali Graphics Debugger home page on ARM's website"))
-					]
-				]
-			]
-		];
-	}
-	// Adreno Profiler settings
-	{
-		TAttribute<EVisibility> AdrenoSettingsVisibility(
-			TAttribute<EVisibility>::Create(TAttribute<EVisibility>::FGetter::CreateStatic(GraphicsDebuggerSettingsVisibility, EAndroidGraphicsDebugger::Adreno, AndroidGraphicsDebuggerProperty))
-		);
-
-		FText AdrenoHelpText = GetAdrenoProfilerHelpText();
-		
-		GraphicsDebuggerCategory.AddCustomRow(LOCTEXT("AdrenoProfilerInfo", "Adreno Profiler Info"), false)
-		.Visibility(AdrenoSettingsVisibility)
-		.WholeRowWidget
-		[
-			SNew(SBorder)
-			.Padding(1)
-			[
-				SNew(SVerticalBox)
-				+ SVerticalBox::Slot()
-				.Padding(FMargin(10, 10, 10, 10))
-				.AutoHeight()
-				[
-					SNew(SRichTextBlock)
-					.Text(AdrenoHelpText)
-					.TextStyle(FAppStyle::Get(), "MessageLog")
-					.DecoratorStyleSet(&FAppStyle::Get())
-					.AutoWrapText(true)
-				]
-				+ SVerticalBox::Slot()
-				.AutoHeight()
-				.Padding(FMargin(10, 10, 10, 10))
-				[	
-					SNew(SBox)
-					.HAlign(HAlign_Left)
-					[
-						SNew(SHyperlinkLaunchURL, TEXT("https://developer.qualcomm.com/software/adreno-gpu-profiler"))
-						.Text(LOCTEXT("AdrenoProfilerPage", "Adreno Profiler home page"))
-						.ToolTipText(LOCTEXT("AdrenoProfilerPageTooltip", "Opens the Adreno Profiler home page on the Qualcomm website"))
-					]
-				]
-			]
-		];
-	}
-}
-
 //////////////////////////////////////////////////////////////////////////
 
 #undef LOCTEXT_NAMESPACE
diff --git a/Engine/Source/Developer/Android/AndroidPlatformEditor/Private/AndroidTargetSettingsCustomization.h b/Engine/Source/Developer/Android/AndroidPlatformEditor/Private/AndroidTargetSettingsCustomization.h
--- a/Engine/Source/Developer/Android/AndroidPlatformEditor/Private/AndroidTargetSettingsCustomization.h
+++ b/Engine/Source/Developer/Android/AndroidPlatformEditor/Private/AndroidTargetSettingsCustomization.h
@@ -30,7 +30,6 @@ private:
 	void BuildAppManifestSection(IDetailLayoutBuilder& DetailLayout);
 	void BuildIconSection(IDetailLayoutBuilder& DetailLayout);
 	void BuildLaunchImageSection(IDetailLayoutBuilder& DetailLayout);
-	void BuildGraphicsDebuggerSection(IDetailLayoutBuilder& DetailLayout);
 
 	// Navigates to the build files in the explorer or finder
 	FReply OpenBuildFolder();
diff --git a/Engine/Source/Programs/UnrealBuildTool/Platform/Android/UEDeployAndroid.cs b/Engine/Source/Programs/UnrealBuildTool/Platform/Android/UEDeployAndroid.cs
--- a/Engine/Source/Programs/UnrealBuildTool/Platform/Android/UEDeployAndroid.cs
+++ b/Engine/Source/Programs/UnrealBuildTool/Platform/Android/UEDeployAndroid.cs
@@ -1407,60 +1407,6 @@ namespace UnrealBuildTool
 			CopySO(FinalSOName, SourceSOName);
 		}
 
-		private void CopyGfxDebugger(string UnrealBuildPath, string NDKArch)
-		{
-			string AndroidGraphicsDebugger;
-			ConfigHierarchy Ini = GetConfigCacheIni(ConfigHierarchyType.Engine);
-			Ini.GetString("/Script/AndroidRuntimeSettings.AndroidRuntimeSettings", "AndroidGraphicsDebugger", out AndroidGraphicsDebugger);
-
-			switch (AndroidGraphicsDebugger.ToLower())
-			{
-				case "mali":
-					{
-						string MaliGraphicsDebuggerPath;
-						AndroidPlatformSDK.GetPath(Ini, "/Script/AndroidRuntimeSettings.AndroidRuntimeSettings", "MaliGraphicsDebuggerPath", out MaliGraphicsDebuggerPath);
-						if (Directory.Exists(MaliGraphicsDebuggerPath))
-						{
-							Directory.CreateDirectory(Path.Combine(UnrealBuildPath, "libs", NDKArch));
-							string MaliLibSrcPath = Path.Combine(MaliGraphicsDebuggerPath, "target", "android-non-root", "arm", NDKArch, "libMGD.so");
-							if (!File.Exists(MaliLibSrcPath))
-							{
-								// in v4.3.0 library location was changed
-								MaliLibSrcPath = Path.Combine(MaliGraphicsDebuggerPath, "target", "android", "arm", "unrooted", NDKArch, "libMGD.so");
-							}
-							string MaliLibDstPath = Path.Combine(UnrealBuildPath, "libs", NDKArch, "libMGD.so");
-
-							Logger.LogInformation("Copying {MaliLibSrcPath} to {MaliLibDstPath}", MaliLibSrcPath, MaliLibDstPath);
-							File.Copy(MaliLibSrcPath, MaliLibDstPath, true);
-							File.SetLastWriteTimeUtc(MaliLibDstPath, File.GetLastWriteTimeUtc(MaliLibSrcPath));
-
-							string MaliVkLayerLibSrcPath = Path.Combine(MaliGraphicsDebuggerPath, "target", "android", "arm", "rooted", NDKArch, "libGLES_aga.so");
-							if (File.Exists(MaliVkLayerLibSrcPath))
-							{
-								string MaliVkLayerLibDstPath = Path.Combine(UnrealBuildPath, "libs", NDKArch, "libVkLayerAGA.so");
-								Logger.LogInformation("Copying {MaliVkLayerLibSrcPath} to {MaliVkLayerLibDstPath}", MaliVkLayerLibSrcPath, MaliVkLayerLibDstPath);
-								File.Copy(MaliVkLayerLibSrcPath, MaliVkLayerLibDstPath, true);
-								File.SetLastWriteTimeUtc(MaliVkLayerLibDstPath, File.GetLastWriteTimeUtc(MaliVkLayerLibSrcPath));
-							}
-						}
-					}
-					break;
-
-				// @TODO: Add NVIDIA Gfx Debugger
-				/*
-				case "nvidia":
-					{
-						Directory.CreateDirectory(UnrealBuildPath + "/libs/" + NDKArch);
-						File.Copy("F:/NVPACK/android-kk-egl-t124-a32/Stripped_libNvPmApi.Core.so", UnrealBuildPath + "/libs/" + NDKArch + "/libNvPmApi.Core.so", true);
-						File.Copy("F:/NVPACK/android-kk-egl-t124-a32/Stripped_libNvidia_gfx_debugger.so", UnrealBuildPath + "/libs/" + NDKArch + "/libNvidia_gfx_debugger.so", true);
-					}
-					break;
-				*/
-				default:
-					break;
-			}
-		}
-
 		void LogBuildSetup()
 		{
 			ConfigHierarchy Ini = GetConfigCacheIni(ConfigHierarchyType.Engine);
@@ -2689,7 +2635,6 @@ namespace UnrealBuildTool
 			bool bPackageForMetaQuest = IsPackagingForMetaQuest(CookFlavor, Ini);
 			Ini.GetBool("OnlineSubsystemGooglePlay.Store", "bSupportsInAppPurchasing", out bool bEnableIAP);
 			Ini.GetBool("/Script/AndroidRuntimeSettings.AndroidRuntimeSettings", "bShowLaunchImage", out bool bShowLaunchImage);
-			Ini.GetString("/Script/AndroidRuntimeSettings.AndroidRuntimeSettings", "AndroidGraphicsDebugger", out string AndroidGraphicsDebugger);
 			Ini.GetBool("/Script/AndroidRuntimeSettings.AndroidRuntimeSettings", "bSupportAdMob", out bool bSupportAdMob);
 			Ini.GetBool("/Script/AndroidRuntimeSettings.AndroidRuntimeSettings", "bValidateTextureFormats", out bool bValidateTextureFormats);
 
@@ -3169,14 +3114,6 @@ namespace UnrealBuildTool
 						Text.AppendLine("\t" + Line);
 					}
 				}
-				if (AndroidGraphicsDebugger.ToLower() == "adreno")
-				{
-					string PermissionString = "\t<uses-permission android:name=\"com.qti.permission.PROFILER\"/>";
-					if (!Text.ToString().Contains(PermissionString, StringComparison.Ordinal))
-					{
-						Text.AppendLine(PermissionString);
-					}
-				}
 
 				if (!bSupportingAllTextureFormats)
 				{
@@ -5206,7 +5143,6 @@ popd
 				}
 				CopySwappy(UnrealBuildPath, NDKArch, Configuration);
 				CopyPSOService(UnrealBuildPath, UnrealPreBuiltFilesPath, NDKArch);
-				CopyGfxDebugger(UnrealBuildPath, NDKArch);
 				CopyVulkanValidationLayers(UnrealBuildPath, NDKArch, Configuration.ToString());
 
 				if (Sanitizer != AndroidToolChain.ClangSanitizer.None)
@@ -6137,22 +6073,6 @@ import java.util.Collection;
 				LoadUnrealDefault = @"		System.loadLibrary(""Unreal"");";
 			}
 
-			string AndroidGraphicsDebugger;
-			Ini.GetString("/Script/AndroidRuntimeSettings.AndroidRuntimeSettings", "AndroidGraphicsDebugger", out AndroidGraphicsDebugger);
-			switch (AndroidGraphicsDebugger.ToLower())
-			{
-				case "mali":
-					LoadLibraryDefaults += "\t\ttry\n" +
-											"\t\t{\n" +
-											"\t\t\tSystem.loadLibrary(\"MGD\");\n" +
-											"\t\t}\n" +
-											"\t\tcatch (java.lang.UnsatisfiedLinkError e)\n" +
-											"\t\t{\n" +
-											"\t\t\tLog.debug(\"libMGD.so not loaded.\");\n" +
-											"\t\t}\n";
-					break;
-			}
-
 			if (bJavaPerfettoMarkers)
 			{
 				GameActivityImportAdditionsDefault += "import androidx.tracing.Trace;\n";
diff --git a/Engine/Source/Runtime/Android/AndroidRuntimeSettings/Classes/AndroidRuntimeSettings.h b/Engine/Source/Runtime/Android/AndroidRuntimeSettings/Classes/AndroidRuntimeSettings.h
--- a/Engine/Source/Runtime/Android/AndroidRuntimeSettings/Classes/AndroidRuntimeSettings.h
+++ b/Engine/Source/Runtime/Android/AndroidRuntimeSettings/Classes/AndroidRuntimeSettings.h
@@ -191,7 +191,7 @@ namespace EAndroidAudio
 UENUM()
 namespace EAndroidGraphicsDebugger
 {
-	enum Type : int
+	enum UE_DEPRECATED(5.8, "Please don't use") Type : int
 	{
 		None = 0 UMETA(DisplayName = "None"),
 		Mali = 1 UMETA(DisplayName = "Mali Graphics Debugger", ToolTip = "Configure for Mali Graphics Debugger."),
@@ -683,12 +683,16 @@ public:
 	UPROPERTY(GlobalConfig)
 	float AutoStreamingThreshold;
 
-	// Several Android graphics debuggers require configuration changes to be made to your application in order to operate. Choosing an option from this menu will configure your project to work with that graphics debugger. 
-	UPROPERTY(GlobalConfig, EditAnywhere, Category = GraphicsDebugger)
+	PRAGMA_DISABLE_DEPRECATION_WARNINGS
+	/** DEPRECATED 5.8 */ 
+	UE_DEPRECATED(5.8, "This option has been removed, please don't use")
+	UPROPERTY(meta = (DeprecatedProperty, DeprecationMessage = "This option has been removed, please don't use"))
 	TEnumAsByte<EAndroidGraphicsDebugger::Type> AndroidGraphicsDebugger;
+	PRAGMA_ENABLE_DEPRECATION_WARNINGS
 
-	/** The path to your Mali Graphics Debugger installation (eg C:/Program Files/ARM/Mali Developer Tools/Mali Graphics Debugger v4.2.0) */
-	UPROPERTY(GlobalConfig, EditAnywhere, Category = GraphicsDebugger)
+	/** DEPRECATED 5.8 */
+	UE_DEPRECATED(5.8, "This option has been removed, please don't use")
+	UPROPERTY(meta = (DeprecatedProperty, DeprecationMessage = "This option has been removed, please don't use"))
 	FDirectoryPath MaliGraphicsDebuggerPath;
 
 	// If checked, this will disable the security.perf_harden flag on the Android device when launching or installing your app via the generated batch file.
